[UEMOB-405] - Mobile PIE polishing - Bezel & Rotation + related subtasks

diff --git a/Engine/Source/Editor/PIEPreviewDeviceProfileSelector/PIEPreviewDeviceProfileSelector.Build.cs b/Engine/Source/Editor/PIEPreviewDeviceProfileSelector/PIEPreviewDeviceProfileSelector.Build.cs
--- a/Engine/Source/Editor/PIEPreviewDeviceProfileSelector/PIEPreviewDeviceProfileSelector.Build.cs
+++ b/Engine/Source/Editor/PIEPreviewDeviceProfileSelector/PIEPreviewDeviceProfileSelector.Build.cs
@@ -13,7 +13,7 @@ namespace UnrealBuildTool.Rules
 				{
 					"Core",
 				    "CoreUObject",
-				}
+                }
 				);
 
 			PrivateDependencyModuleNames.AddRange(
@@ -24,8 +24,12 @@ namespace UnrealBuildTool.Rules
 					"JsonUtilities",
 					"MaterialShaderQualitySettings",
 					"Slate",
-					"SlateCore"
-				}
+					"SlateCore",
+                    "ApplicationCore",
+					"Engine",
+					"UnrealEd",
+					"EditorStyle"
+                }
 				);
 		}
 	}
diff --git a/Engine/Source/Editor/PIEPreviewDeviceProfileSelector/Private/PIEPreviewDevice.cpp b/Engine/Source/Editor/PIEPreviewDeviceProfileSelector/Private/PIEPreviewDevice.cpp
new file mode 100644
--- /dev/null
+++ b/Engine/Source/Editor/PIEPreviewDeviceProfileSelector/Private/PIEPreviewDevice.cpp
@@ -0,0 +1,415 @@
+// Copyright 1998-2018 Epic Games, Inc. All Rights Reserved.
+
+#include "PIEPreviewDevice.h"
+
+#include "HAL/IConsoleManager.h"
+#include "CoreGlobals.h"
+#include "RHI.h"
+#include "Misc/ConfigCacheIni.h"
+#include "IDeviceProfileSelectorModule.h"
+#include "Modules/ModuleManager.h"
+#include "MaterialShaderQualitySettings.h"
+#include "Android/AndroidWindowUtils.h"
+#include "GenericPlatform/GenericApplication.h"
+#include "Engine/Texture2D.h"
+#include "ImageUtils.h"
+#include "HAL/PlatformApplicationMisc.h"
+#include "SEditorViewport.h"
+
+static const int32 g_JSON_VALUE_NOT_SET = 0;
+
+FPIEPreviewDevice::FPIEPreviewDevice()
+{
+	DeviceSpecs = MakeShareable(new FPIEPreviewDeviceSpecifications());
+}
+
+void FPIEPreviewDevice::ComputeViewportSize(const bool bClampWindowSize)
+{
+	int32 ScreenWidth, ScreenHeight;
+	ComputeDeviceResolution(ScreenWidth, ScreenHeight);
+
+	FPIEPreviewDeviceBezelViewportRect ViewportRect = DeviceSpecs->BezelProperties.BezelViewportRect;
+
+	if (IsDeviceFlipped())
+	{
+		Swap(ScreenWidth, ScreenHeight);
+		Swap(ViewportRect.X, ViewportRect.Y);
+		Swap(ViewportRect.Width, ViewportRect.Height);
+	}
+
+	float ScaleX = (float)ScreenWidth / (float)ViewportRect.Width;
+	float ScaleY = (float)ScreenHeight / (float)ViewportRect.Height;
+
+	float BezelScaleFactor = 1.0f / DPIScaleFactor;
+
+ 	// compute widow size
+	WindowWidth = ScreenWidth;
+	WindowHeight = ScreenHeight + WindowTitleBarSize * DPIScaleFactor;
+
+	// compute viewport margin
+	if (bShowBezel && BezelTexture != nullptr)
+	{
+		// compute widow size
+		WindowWidth += 2.0f * ViewportRect.X * ScaleX;
+		WindowHeight += 2.0f * ViewportRect.Y * ScaleY;
+
+		ViewportRect.X = FMath::RoundToInt(ViewportRect.X * BezelScaleFactor);
+		ViewportRect.Y = FMath::RoundToInt(ViewportRect.Y * BezelScaleFactor);
+		ViewportRect.Width = FMath::RoundToInt(ViewportRect.Width * BezelScaleFactor);
+		ViewportRect.Height = FMath::RoundToInt(ViewportRect.Height * BezelScaleFactor);
+
+		ViewportMargin.Left = ViewportRect.X;
+		ViewportMargin.Top = ViewportRect.Y;
+
+		int32 BezelWidth = IsDeviceFlipped() ? BezelTexture->GetSizeY() : BezelTexture->GetSizeX();
+		int32 BezelHeight = IsDeviceFlipped() ? BezelTexture->GetSizeX() : BezelTexture->GetSizeY();
+
+		BezelWidth = FMath::RoundToInt(BezelWidth * BezelScaleFactor);
+		BezelHeight = FMath::RoundToInt(BezelHeight * BezelScaleFactor);
+
+		ViewportMargin.Right = (float)(BezelWidth - ViewportRect.Width - ViewportRect.X);
+		ViewportMargin.Bottom = (float)(BezelHeight - ViewportRect.Height - ViewportRect.Y);
+
+		ViewportMargin = ViewportMargin * FMargin(ScaleX, ScaleY);
+	}
+	else
+	{
+		ViewportMargin = FMargin(0.0f);
+	}
+
+	// if necessary constrain the window inside the desktop boundaries
+	if (bClampWindowSize)
+	{
+		FDisplayMetrics DisplayMetrics;
+		FDisplayMetrics::GetDisplayMetrics(DisplayMetrics);
+
+		auto DesktopWidth = DisplayMetrics.VirtualDisplayRect.Right - DisplayMetrics.VirtualDisplayRect.Left;
+		auto DesktopHeight = DisplayMetrics.VirtualDisplayRect.Bottom - DisplayMetrics.VirtualDisplayRect.Top;
+
+		if (WindowWidth > DesktopWidth)
+		{
+			float ScaleFactor = (float)DesktopWidth / (float)WindowWidth;
+			ScaleX *= ScaleFactor;
+			ScaleY *= ScaleFactor;
+
+			WindowWidth = DesktopWidth;
+			WindowHeight *= ScaleFactor;
+		}
+		if (WindowHeight > DesktopHeight)
+		{
+			float ScaleFactor = (float)DesktopHeight / (float)WindowHeight;
+			ScaleX *= ScaleFactor;
+			ScaleY *= ScaleFactor;
+
+			WindowWidth *= ScaleFactor;
+			WindowHeight = DesktopHeight;
+		}
+	}
+}
+
+void FPIEPreviewDevice::GetDeviceDefaultResolution(int32& Width, int32& Height)
+{
+	Width = DeviceSpecs->ResolutionX;
+
+	if (DeviceSpecs->ResolutionYImmersiveMode != g_JSON_VALUE_NOT_SET)
+	{
+		DeviceSpecs->ResolutionY = DeviceSpecs->ResolutionYImmersiveMode;
+	}
+
+	Height = DeviceSpecs->ResolutionY;
+}
+
+void FPIEPreviewDevice::ComputeDeviceResolution(int32& Width, int32& Height)
+{
+	GetDeviceDefaultResolution(Width, Height);
+
+	if (!bIgnoreContentScaleFactor)
+	{
+		switch (DeviceSpecs->DevicePlatform)
+		{
+			case EPIEPreviewDeviceType::Android:
+			{
+				ERHIFeatureLevel::Type DeviceFeatureLevel = GetPreviewDeviceFeatureLevel();
+				if (DeviceFeatureLevel < ERHIFeatureLevel::ES3_1)
+				{
+					AndroidWindowUtils::ApplyMosaicRequirements(Width, Height);
+				}
+
+				AndroidWindowUtils::ApplyContentScaleFactor(Width, Height);
+			}
+			break;
+
+			case EPIEPreviewDeviceType::IOS:
+			{
+				static IConsoleVariable* CVar = IConsoleManager::Get().FindConsoleVariable(TEXT("r.MobileContentScaleFactor"));
+				float RequestedContentScaleFactor = CVar->GetFloat();
+
+				if (FMath::IsNearlyZero(RequestedContentScaleFactor))
+				{
+					RequestedContentScaleFactor = DeviceSpecs->IOSProperties.NativeScaleFactor;
+				}
+
+				Width *= RequestedContentScaleFactor;
+				Height *= RequestedContentScaleFactor;
+			}
+			break;
+
+			default:
+				break;
+		} //end switch
+	}// end if (!bIgnoreContentScaleFactor)
+
+	Width *= ResolutionScaleFactor;
+	Height *= ResolutionScaleFactor;
+}
+
+void FPIEPreviewDevice::DetermineScreenOrientationRequirements(bool& bNeedPortrait, bool& bNeedLandscape)
+{
+	switch (DeviceSpecs->DevicePlatform)
+	{
+		case EPIEPreviewDeviceType::Android:
+		{
+			FString Orientation;
+			GConfig->GetString(TEXT("/Script/AndroidRuntimeSettings.AndroidRuntimeSettings"), TEXT("Orientation"), Orientation, GEngineIni);
+			if (Orientation.ToLower().Equals("portrait") || Orientation.ToLower().Equals("reverseportrait") || Orientation.ToLower().Equals("sensorportrait"))
+			{
+				bNeedPortrait = true;
+			}
+			else if (Orientation.ToLower().Equals("landscape") || Orientation.ToLower().Equals("reverselandscape") || Orientation.ToLower().Equals("sensorlandscape"))
+			{
+				bNeedLandscape = true;
+			}
+			else if (Orientation.ToLower().Equals("sensor") || Orientation.ToLower().Equals("fullsensor"))
+			{
+				bNeedPortrait = true;
+				bNeedLandscape = true;
+			}
+			else
+			{
+				bNeedPortrait = true;
+				bNeedLandscape = true;
+			}
+		}
+		break;
+
+		case EPIEPreviewDeviceType::IOS:
+		{
+			bool bSupportsPortraitOrientation;
+			bool bSupportsUpsideDownOrientation;
+			bool bSupportsLandscapeLeftOrientation;
+			bool bSupportsLandscapeRightOrientation;
+			GConfig->GetBool(TEXT("/Script/IOSRuntimeSettings.IOSRuntimeSettings"), TEXT("bSupportsPortraitOrientation"), bSupportsPortraitOrientation, GEngineIni);
+			GConfig->GetBool(TEXT("/Script/IOSRuntimeSettings.IOSRuntimeSettings"), TEXT("bSupportsUpsideDownOrientation"), bSupportsUpsideDownOrientation, GEngineIni);
+			GConfig->GetBool(TEXT("/Script/IOSRuntimeSettings.IOSRuntimeSettings"), TEXT("bSupportsLandscapeLeftOrientation"), bSupportsLandscapeLeftOrientation, GEngineIni);
+			GConfig->GetBool(TEXT("/Script/IOSRuntimeSettings.IOSRuntimeSettings"), TEXT("bSupportsLandscapeRightOrientation"), bSupportsLandscapeRightOrientation, GEngineIni);
+			bNeedPortrait = bSupportsPortraitOrientation || bSupportsUpsideDownOrientation;
+			bNeedLandscape = bSupportsLandscapeLeftOrientation || bSupportsLandscapeRightOrientation;
+		}
+		break;
+	}
+}
+
+ERHIFeatureLevel::Type FPIEPreviewDevice::GetPreviewDeviceFeatureLevel() const
+{
+	switch (DeviceSpecs->DevicePlatform)
+	{
+		case EPIEPreviewDeviceType::Android:
+		{
+			FString SubVersion;
+			// Check for ES3.1+ support from GLVersion, TODO: check other ES31 feature level constraints, see android's PlatformInitOpenGL
+			const bool bDeviceSupportsES31 = DeviceSpecs->AndroidProperties.GLVersion.Split(TEXT("OpenGL ES 3."), nullptr, &SubVersion) && FCString::Atoi(*SubVersion) >= 1;
+
+			// check the project's gles support:
+			bool bProjectBuiltForES2 = false, bProjectBuiltForES31 = false;
+			GConfig->GetBool(TEXT("/Script/AndroidRuntimeSettings.AndroidRuntimeSettings"), TEXT("bBuildForES31"), bProjectBuiltForES31, GEngineIni);
+			GConfig->GetBool(TEXT("/Script/AndroidRuntimeSettings.AndroidRuntimeSettings"), TEXT("bBuildForES2"), bProjectBuiltForES2, GEngineIni);
+
+			// Android Preview Device is currently expected to work on gles.
+			check(bProjectBuiltForES2 || bProjectBuiltForES31);
+
+			// Projects without ES2 support can only expect to run on ES31 devices.
+			check(bProjectBuiltForES2 || bDeviceSupportsES31);
+
+			// ES3.1+ devices fallback to ES2 if the project itself doesn't support ES3.1
+			return bDeviceSupportsES31 && bProjectBuiltForES31 ? ERHIFeatureLevel::ES3_1 : ERHIFeatureLevel::ES2;
+		}
+		case EPIEPreviewDeviceType::IOS:
+		{
+			bool bProjectBuiltForMetal = false, bProjectBuiltForMRTMetal = false;
+			GConfig->GetBool(TEXT("/Script/IOSRuntimeSettings.IOSRuntimeSettings"), TEXT("bSupportsMetal"), bProjectBuiltForMetal, GEngineIni);
+			GConfig->GetBool(TEXT("/Script/IOSRuntimeSettings.IOSRuntimeSettings"), TEXT("bSupportsMetalMRT"), bProjectBuiltForMRTMetal, GEngineIni);
+
+			const bool bDeviceSupportsMetal = DeviceSpecs->IOSProperties.MetalRHIState.MaxTextureDimensions > 0;
+
+			// not supporting preview for MRT metal 
+			check(!bProjectBuiltForMRTMetal);
+
+			// at least one of these should be valid!
+			check(bProjectBuiltForMetal);
+
+			// if device doesn't support metal the project must have ES2 enabled.
+			check(bProjectBuiltForMetal && bDeviceSupportsMetal);
+
+			return ERHIFeatureLevel::ES3_1;
+		}
+	}
+
+	checkNoEntry();
+	return  ERHIFeatureLevel::Num;
+}
+
+void FPIEPreviewDevice::SetupDevice(const int32 InWindowTitleBarSize)
+{
+	WindowTitleBarSize = InWindowTitleBarSize;
+
+	// set initial scale factor
+	ResolutionScaleFactor = 0.5f;
+
+	// compute bezel file path
+	FString BezelPath = FPaths::EngineContentDir() + TEXT("Editor/PIEPreviewDeviceSpecs/");
+	if (DeviceSpecs->DevicePlatform == EPIEPreviewDeviceType::Android)
+	{
+		BezelPath += TEXT("Android/");
+	}
+	else if (DeviceSpecs->DevicePlatform == EPIEPreviewDeviceType::IOS)
+	{
+		BezelPath += TEXT("iOS/");
+	}
+	BezelPath += DeviceSpecs->BezelProperties.DeviceBezelFile;
+
+	// load the bezel texture
+	BezelTexture = FImageUtils::ImportFileAsTexture2D(BezelPath);
+	
+	// if we have invalid/uninitialized viewport values use the values provided as native device resolution
+	FPIEPreviewDeviceBezelViewportRect& ViewportRect = DeviceSpecs->BezelProperties.BezelViewportRect;
+	if (BezelTexture == nullptr || ViewportRect.Width == 0 || ViewportRect.Height == 0)
+	{
+		ViewportRect.X = 0;
+		ViewportRect.Y = 0;
+		ViewportRect.Width = GetDeviceSpecs()->ResolutionX;
+		ViewportRect.Height = GetDeviceSpecs()->ResolutionY;
+	}
+
+	// check rotation functionalities
+	bool bPortrait = false, bLandscape = false;
+	DetermineScreenOrientationRequirements(bPortrait, bLandscape);
+
+	bAllowRotation = (bPortrait && bLandscape);
+
+	// determine current orientation
+	bool bSwitchOrientation = bLandscape && DeviceSpecs->ResolutionY > DeviceSpecs->ResolutionX;
+	bSwitchOrientation |= !bLandscape && bPortrait && DeviceSpecs->ResolutionX > DeviceSpecs->ResolutionY;
+	if (bSwitchOrientation)
+	{
+		SwitchOrientation(true);
+	}
+	else
+	{
+		// finally compute needed window and viewport size
+		// the above branch will call it from SwitchOrientation()
+		ComputeViewportSize(true);
+	}
+
+	ApplyRHIOverrides();
+}
+
+void FPIEPreviewDevice::ApplyRHIPrerequisitesOverrides() const
+{
+	RHISetMobilePreviewFeatureLevel(GetPreviewDeviceFeatureLevel());
+}
+
+void FPIEPreviewDevice::ApplyRHIOverrides() const
+{
+	EShaderPlatform PreviewPlatform = SP_NumPlatforms;
+	ERHIFeatureLevel::Type PreviewFeatureLevel = GetPreviewDeviceFeatureLevel();
+	FPIERHIOverrideState* RHIOverrideState = nullptr;
+
+	switch (DeviceSpecs->DevicePlatform)
+	{
+		case EPIEPreviewDeviceType::Android:
+		{
+			if (PreviewFeatureLevel == ERHIFeatureLevel::ES2)
+			{
+				PreviewPlatform = SP_OPENGL_ES2_ANDROID;
+				RHIOverrideState = &DeviceSpecs->AndroidProperties.GLES2RHIState;
+			}
+			else
+			{
+				PreviewPlatform = SP_OPENGL_ES3_1_ANDROID;
+				RHIOverrideState = &DeviceSpecs->AndroidProperties.GLES31RHIState;
+			}
+		}
+		break;
+
+		case EPIEPreviewDeviceType::IOS:
+		{
+			PreviewPlatform = SP_METAL_MACES3_1;
+			RHIOverrideState = &DeviceSpecs->IOSProperties.MetalRHIState;
+		}
+		break;
+
+		default:
+		break;
+	}
+
+	if (PreviewPlatform != SP_NumPlatforms)
+	{
+		FName QualityPreviewShaderPlatform = LegacyShaderPlatformToShaderFormat(PreviewPlatform);
+		UMaterialShaderQualitySettings::Get()->SetPreviewPlatform(QualityPreviewShaderPlatform);
+	}
+
+	GMaxTextureDimensions.SetPreviewOverride(RHIOverrideState->MaxTextureDimensions);
+	GMaxShadowDepthBufferSizeX.SetPreviewOverride(RHIOverrideState->MaxShadowDepthBufferSizeX);
+	GMaxShadowDepthBufferSizeY.SetPreviewOverride(RHIOverrideState->MaxShadowDepthBufferSizeY);
+	GMaxCubeTextureDimensions.SetPreviewOverride(RHIOverrideState->MaxCubeTextureDimensions);
+	GRHISupportsInstancing.SetPreviewOverride(RHIOverrideState->SupportsInstancing);
+	GSupportsMultipleRenderTargets.SetPreviewOverride(RHIOverrideState->SupportsMultipleRenderTargets);
+	GSupportsRenderTargetFormat_PF_FloatRGBA.SetPreviewOverride(RHIOverrideState->SupportsRenderTargetFormat_PF_FloatRGBA);
+	GSupportsRenderTargetFormat_PF_G8.SetPreviewOverride(RHIOverrideState->SupportsRenderTargetFormat_PF_G8);
+}
+
+FString FPIEPreviewDevice::GetProfile() const
+{
+	FString Profile;
+
+	switch (DeviceSpecs->DevicePlatform)
+	{
+		case EPIEPreviewDeviceType::Android:
+		{
+			IDeviceProfileSelectorModule* AndroidDeviceProfileSelector = FModuleManager::LoadModulePtr<IDeviceProfileSelectorModule>("AndroidDeviceProfileSelector");
+			if (AndroidDeviceProfileSelector)
+			{
+				FPIEAndroidDeviceProperties& AndroidProperties = DeviceSpecs->AndroidProperties;
+
+				TMap<FString, FString> DeviceParameters;
+				DeviceParameters.Add("GPUFamily", AndroidProperties.GPUFamily);
+				DeviceParameters.Add("GLVersion", AndroidProperties.GLVersion);
+				DeviceParameters.Add("VulkanAvailable", AndroidProperties.VulkanAvailable ? "true" : "false");
+				DeviceParameters.Add("VulkanVersion", AndroidProperties.VulkanVersion);
+				DeviceParameters.Add("AndroidVersion", AndroidProperties.AndroidVersion);
+				DeviceParameters.Add("DeviceMake", AndroidProperties.DeviceMake);
+				DeviceParameters.Add("DeviceModel", AndroidProperties.DeviceModel);
+				DeviceParameters.Add("UsingHoudini", AndroidProperties.UsingHoudini ? "true" : "false");
+
+				FString PIEProfileName = AndroidDeviceProfileSelector->GetDeviceProfileName(DeviceParameters);
+				if (!PIEProfileName.IsEmpty())
+				{
+					Profile = PIEProfileName;
+				}
+			}
+		}
+		break;
+
+		case EPIEPreviewDeviceType::IOS:
+		{
+			FPIEIOSDeviceProperties& IOSProperties = DeviceSpecs->IOSProperties;
+			Profile = IOSProperties.DeviceModel;
+		}
+		break;
+
+		default:
+		break;
+	}
+
+	return Profile;
+}
\ No newline at end of file
diff --git a/Engine/Source/Editor/PIEPreviewDeviceProfileSelector/Private/PIEPreviewDevice.h b/Engine/Source/Editor/PIEPreviewDeviceProfileSelector/Private/PIEPreviewDevice.h
new file mode 100644
--- /dev/null
+++ b/Engine/Source/Editor/PIEPreviewDeviceProfileSelector/Private/PIEPreviewDevice.h
@@ -0,0 +1,180 @@
+// Copyright 1998-2018 Epic Games, Inc. All Rights Reserved.
+
+#pragma once
+
+#include "CoreMinimal.h"
+#include "RHIDefinitions.h"
+#include "Layout/Margin.h"
+#include "PIEPreviewDeviceSpecification.h"
+
+/**
+ * FPIEPreviewDevice class stores specific device settings (FPIEPreviewDeviceSpecifications) and, on request, applies them to the Unreal runtime system
+*/
+class FPIEPreviewDevice
+{
+private:
+	TSharedPtr<FPIEPreviewDeviceSpecifications> DeviceSpecs;
+
+	// the required size of the display window
+	int32 WindowWidth = 0;
+	int32 WindowHeight = 0;
+
+	// window title bar size, needed to compute the final window size
+	int32 WindowTitleBarSize = 0;
+
+	// margins used to correctly position the viewport widget inside the provided bezel
+	FMargin ViewportMargin;
+
+	// whether or not this device can rotate its screen
+	bool bAllowRotation = true;
+
+	// true if we are in a 'rotated' state from the original orientation as provided in the json file
+	bool bDeviceFlipped = false;
+
+	// global scaling factor applied to the whole window
+	float ResolutionScaleFactor = 1.0f;
+
+	// DPI scale factor used in window size computations
+	float DPIScaleFactor = 1.0f;
+
+	/** when true extra window space will be allocated and the phone bezel will be rendered */
+	bool bShowBezel = true;
+
+	/** when true r.MobileContentScaleFactor is ignored */
+	bool bIgnoreContentScaleFactor = false;
+
+	UPROPERTY()
+	class UTexture2D* BezelTexture = nullptr;
+
+private:
+	// utility function that computes the viewport resolution
+	void ComputeDeviceResolution(int32& ScreenWidth, int32& ScreenHeight);
+
+	// utility function that will attempt to determine the supported device orientations
+	void DetermineScreenOrientationRequirements(bool& bNeedPortrait, bool& bNeedLandscape);
+
+	// this should provide the needed RHI feature level based on specified json parameters and Unreal system settings
+	ERHIFeatureLevel::Type GetPreviewDeviceFeatureLevel() const;
+
+	// call this function to apply specific RHI settings as specified in the json file
+	void ApplyRHIOverrides() const;
+
+	// function that computes the viewport widget offset and needed window size
+	void ComputeViewportSize(const bool bClampWindowSize);
+
+public:
+	FPIEPreviewDevice();
+
+	// call this to run device setup -> load the bezel texture, compute appropriate orientation and apply device specific RHI settings
+	void SetupDevice(const int32 InWindowTitleBarSize);
+
+	// call this before RHI creation to apply needed setup overrides
+	void ApplyRHIPrerequisitesOverrides() const;
+
+	TSharedPtr<FPIEPreviewDeviceSpecifications> GetDeviceSpecs();
+
+	// call this to retrieve the viewport widget padding inside the bezel 
+	FMargin GetViewportMargin() const;
+
+	// functions to retrieve the needed window size
+	int32 GetWindowWidth() const;
+	int32 GetWindowHeight() const;
+
+	// returns true if the current device was rotated from the original orientation
+	bool IsDeviceFlipped() const;
+
+	// call these functions to flag device rotation and specific windows scaling
+	// a call to ComputeViewportSize() is necessary to update the window size
+	void SwitchOrientation(const bool bClampWindowSize);
+	void ScaleResolution(const float ScreenFactor, const float DPIFactor, const bool bClampWindowSize);
+	void SetBezelVisibility(const bool bBezelVisible, const bool bClampWindowSize);
+
+	bool GetBezelVisibility() const;
+
+	// returns the current resolution scale factor (float ResolutionScaleFactor)
+	float GetResolutionScale() const;
+
+	// returns the default device resolution, as specified in the json file
+	void GetDeviceDefaultResolution(int32& ScreenWidth, int32& ScreenHeight);
+
+	// if true the device supports rotations
+	bool IsRotationAllowed() const;
+
+	void SetIgnoreMobileContentScaleFactor(bool bIgnore);
+
+	FString GetProfile() const;
+
+	UTexture2D *GetBezelTexture();
+};
+
+FORCEINLINE TSharedPtr<FPIEPreviewDeviceSpecifications> FPIEPreviewDevice::GetDeviceSpecs()
+{
+	return DeviceSpecs;
+}
+
+FORCEINLINE FMargin FPIEPreviewDevice::GetViewportMargin() const
+{
+	return ViewportMargin;
+}
+
+FORCEINLINE int32 FPIEPreviewDevice::GetWindowWidth() const
+{
+	return WindowWidth;
+}
+
+FORCEINLINE int32 FPIEPreviewDevice::GetWindowHeight() const
+{
+	return WindowHeight;
+}
+
+FORCEINLINE bool FPIEPreviewDevice::IsDeviceFlipped() const
+{
+	return bDeviceFlipped;
+}
+
+FORCEINLINE bool FPIEPreviewDevice::IsRotationAllowed() const
+{
+	return bAllowRotation;
+}
+
+FORCEINLINE void FPIEPreviewDevice::SwitchOrientation(const bool bClampWindowSize)
+{
+	bDeviceFlipped = !bDeviceFlipped;
+
+	ComputeViewportSize(bClampWindowSize);
+}
+
+FORCEINLINE void FPIEPreviewDevice::ScaleResolution(const float ScreenFactor, const float DPIFactor, const bool bClampWindowSize)
+{
+	ResolutionScaleFactor = ScreenFactor;
+	DPIScaleFactor = DPIFactor;
+
+	ComputeViewportSize(bClampWindowSize);
+}
+
+FORCEINLINE void FPIEPreviewDevice::SetBezelVisibility(const bool bBezelVisible, const bool bClampWindowSize)
+{
+	bShowBezel = bBezelVisible;
+
+	ComputeViewportSize(bClampWindowSize);
+}
+
+FORCEINLINE bool FPIEPreviewDevice::GetBezelVisibility() const
+{
+	return bShowBezel;
+}
+
+FORCEINLINE float FPIEPreviewDevice::GetResolutionScale() const
+{
+	return ResolutionScaleFactor;
+}
+
+FORCEINLINE UTexture2D *FPIEPreviewDevice::GetBezelTexture()
+{
+	return BezelTexture;
+}
+
+FORCEINLINE void FPIEPreviewDevice::SetIgnoreMobileContentScaleFactor(bool bIgnore)
+{
+	bIgnoreContentScaleFactor = bIgnore;
+}
\ No newline at end of file
diff --git a/Engine/Source/Editor/PIEPreviewDeviceProfileSelector/Private/PIEPreviewDeviceProfileSelectorModule.cpp b/Engine/Source/Editor/PIEPreviewDeviceProfileSelector/Private/PIEPreviewDeviceProfileSelectorModule.cpp
--- a/Engine/Source/Editor/PIEPreviewDeviceProfileSelector/Private/PIEPreviewDeviceProfileSelectorModule.cpp
+++ b/Engine/Source/Editor/PIEPreviewDeviceProfileSelector/Private/PIEPreviewDeviceProfileSelectorModule.cpp
@@ -15,6 +15,10 @@
 #include "Misc/ConfigCacheIni.h"
 #include "Internationalization/Culture.h"
 #include "PIEPreviewWindowStyle.h"
+#include "PIEPreviewDevice.h"
+#include "PIEPreviewWindow.h"
+#include "Framework/Application/SlateApplication.h"
+#include "UnrealEngine.h"
 
 DECLARE_LOG_CATEGORY_EXTERN(LogPIEPreviewDevice, Log, All); 
 DEFINE_LOG_CATEGORY(LogPIEPreviewDevice);
@@ -44,213 +48,18 @@ void FPIEPreviewDeviceModule::InitPreviewDevice()
 
 	if (ReadDeviceSpecification())
 	{
-		ERHIFeatureLevel::Type PreviewFeatureLevel = ERHIFeatureLevel::Num;
-		switch (DeviceSpecs->DevicePlatform)
-		{
-			case EPIEPreviewDeviceType::Android:
-			{
-				IDeviceProfileSelectorModule* AndroidDeviceProfileSelector = FModuleManager::LoadModulePtr<IDeviceProfileSelectorModule>("AndroidDeviceProfileSelector");
-				if (AndroidDeviceProfileSelector)
-				{
-					FPIEAndroidDeviceProperties& AndroidProperties = DeviceSpecs->AndroidProperties;
-
-					TMap<FString, FString> DeviceParameters;
-					DeviceParameters.Add("GPUFamily", AndroidProperties.GPUFamily);
-					DeviceParameters.Add("GLVersion", AndroidProperties.GLVersion);
-					DeviceParameters.Add("VulkanAvailable", AndroidProperties.VulkanAvailable ? "true" : "false");
-					DeviceParameters.Add("VulkanVersion", AndroidProperties.VulkanVersion);
-					DeviceParameters.Add("AndroidVersion", AndroidProperties.AndroidVersion);
-					DeviceParameters.Add("DeviceMake", AndroidProperties.DeviceMake);
-					DeviceParameters.Add("DeviceModel", AndroidProperties.DeviceModel);
-					DeviceParameters.Add("UsingHoudini", AndroidProperties.UsingHoudini ? "true" : "false");
-
-					FString PIEProfileName = AndroidDeviceProfileSelector->GetDeviceProfileName(DeviceParameters);
-					if (!PIEProfileName.IsEmpty())
-					{
-						DeviceProfile = PIEProfileName;
-					}
-				}
-				break;
-			}
-			case EPIEPreviewDeviceType::IOS:
-			{
-				FPIEIOSDeviceProperties& IOSProperties = DeviceSpecs->IOSProperties;
-				DeviceProfile = IOSProperties.DeviceModel;
-				break;
-			}
-		}
-		RHISetMobilePreviewFeatureLevel(GetPreviewDeviceFeatureLevel());
-	}
-}
-
-static void ApplyRHIOverrides(FPIERHIOverrideState* RHIOverrideState)
-{
-	check(RHIOverrideState);
-	GMaxTextureDimensions.SetPreviewOverride(RHIOverrideState->MaxTextureDimensions);
-	GMaxShadowDepthBufferSizeX.SetPreviewOverride(RHIOverrideState->MaxShadowDepthBufferSizeX);
-	GMaxShadowDepthBufferSizeY.SetPreviewOverride(RHIOverrideState->MaxShadowDepthBufferSizeY);
-	GMaxCubeTextureDimensions.SetPreviewOverride(RHIOverrideState->MaxCubeTextureDimensions);
-	GRHISupportsInstancing.SetPreviewOverride(RHIOverrideState->SupportsInstancing);
-	GSupportsMultipleRenderTargets.SetPreviewOverride(RHIOverrideState->SupportsMultipleRenderTargets);
-	GSupportsRenderTargetFormat_PF_FloatRGBA.SetPreviewOverride(RHIOverrideState->SupportsRenderTargetFormat_PF_FloatRGBA);
-	GSupportsRenderTargetFormat_PF_G8.SetPreviewOverride(RHIOverrideState->SupportsRenderTargetFormat_PF_G8);
-}
-
-const void FPIEPreviewDeviceModule::GetPreviewDeviceResolution(int32& ScreenWidth, int32& ScreenHeight)
-{
-	if (!DeviceSpecs.IsValid()) 
-	{
-		return;
-	}
-	const int32 JSON_VALUE_NOT_SET = 0;
-	ScreenWidth = DeviceSpecs->ResolutionX;
-	if (DeviceSpecs->ResolutionYImmersiveMode != JSON_VALUE_NOT_SET)
-		DeviceSpecs->ResolutionY = DeviceSpecs->ResolutionYImmersiveMode;
-	ScreenHeight = DeviceSpecs->ResolutionY;
-	static IConsoleVariable* CVar = IConsoleManager::Get().FindConsoleVariable(TEXT("r.MobileContentScaleFactor"));
-	float RequestedContentScaleFactor = CVar->GetFloat();
-	switch (DeviceSpecs->DevicePlatform)
-	{
-		case EPIEPreviewDeviceType::Android:
-		{
-			/********************************************/
-			/**code from FAndroidWindow GetScreenRect()**/
-			/***GearVr and Daydream have been left out***/
-			/********************************************/
-
-			// CSF is a multiplier to 1280x720
-			// determine mosaic requirements:
-			static auto* MobileHDRCvar = IConsoleManager::Get().FindTConsoleVariableDataInt(TEXT("r.MobileHDR"));
-			const bool bMobileHDR = (MobileHDRCvar && MobileHDRCvar->GetValueOnAnyThread() == 1);
-
-			static auto* MobileHDR32bppModeCvar = IConsoleManager::Get().FindTConsoleVariableDataInt(TEXT("r.MobileHDR32bppMode"));
-			const int32 MobileHDR32Mode = MobileHDR32bppModeCvar->GetValueOnAnyThread();
-
-			bool bMosaicEnabled = false;
-			bool bHDR32ModeOverridden = false;
-			bool bDeviceRequiresHDR32bpp = false;
-			bool bDeviceRequiresMosaic = false;
-			
-			bDeviceRequiresHDR32bpp = GSupportsRenderTargetFormat_PF_FloatRGBA;// !FAndroidMisc::SupportsFloatingPointRenderTargets();
-			bDeviceRequiresMosaic = bDeviceRequiresHDR32bpp && GSupportsShaderFramebufferFetch/*&& !FAndroidMisc::SupportsShaderFramebufferFetch()*/;
-
-			bHDR32ModeOverridden = MobileHDR32Mode != 0;
-			bMosaicEnabled = bDeviceRequiresMosaic && (!bHDR32ModeOverridden || MobileHDR32Mode == 1);
-			
-			// get the aspect ratio of the physical screen
-			bool GAndroidIsPortrait = false;
-
-			// some phones gave it the other way (so, if swap if the app is landscape, but width < height)
-			if ((GAndroidIsPortrait && ScreenWidth > ScreenHeight) ||
-				(!GAndroidIsPortrait && ScreenWidth < ScreenHeight))
-			{
-				Swap(ScreenWidth, ScreenHeight);
-			}
-
-			// ensure the size is divisible by a specified amount
-			// do not convert to a surface size that is larger than native resolution
-			const int DividableBy = 8;
-			ScreenWidth = (ScreenWidth / DividableBy) * DividableBy;
-			ScreenHeight = (ScreenHeight / DividableBy) * DividableBy;
-			float AspectRatio = (float)ScreenWidth / (float)ScreenHeight;
-
-			int32 MaxWidth = ScreenWidth;
-			int32 MaxHeight = ScreenHeight;
-
-
-			UE_LOG(LogAndroid, Log, TEXT("Mobile HDR: %s"), bMobileHDR ? TEXT("YES") : TEXT("no"));
-			if (bMobileHDR)
-			{
-				UE_LOG(LogAndroid, Log, TEXT("Device requires 32BPP mode : %s"), bDeviceRequiresHDR32bpp ? TEXT("YES") : TEXT("no"));
-				UE_LOG(LogAndroid, Log, TEXT("Device requires mosaic: %s"), bDeviceRequiresMosaic ? TEXT("YES") : TEXT("no"));
-
-				if (bHDR32ModeOverridden)
-				{
-					UE_LOG(LogAndroid, Log, TEXT("--- Enabling 32 BPP override with 'r.MobileHDR32bppMode' = %d"), MobileHDR32Mode);
-					UE_LOG(LogAndroid, Log, TEXT("  32BPP mode : YES"));
-					UE_LOG(LogAndroid, Log, TEXT("  32BPP mode requires mosaic: %s"), bMosaicEnabled ? TEXT("YES") : TEXT("no"));
-					UE_LOG(LogAndroid, Log, TEXT("  32BPP mode requires RGBE: %s"), MobileHDR32Mode == 2 ? TEXT("YES") : TEXT("no"));
-				}
-
-				if (bMosaicEnabled)
-				{
-					UE_LOG(LogAndroid, Log, TEXT("Using mosaic rendering due to lack of Framebuffer Fetch support."));
-					if (GetPreviewDeviceFeatureLevel() == ERHIFeatureLevel::ES3_1)
-					{
-						const int32 OldMaxWidth = MaxWidth;
-						const int32 OldMaxHeight = MaxHeight;
-
-						if (GAndroidIsPortrait)
-						{
-							MaxHeight = FPlatformMath::Min(MaxHeight, 1024);
-							MaxWidth = MaxHeight * AspectRatio;
-						}
-						else
-						{
-							MaxWidth = FPlatformMath::Min(MaxWidth, 1024);
-							MaxHeight = MaxWidth / AspectRatio;
-						}
-
-						UE_LOG(LogAndroid, Log, TEXT("Limiting MaxWidth=%d and MaxHeight=%d due to mosaic rendering on ES2 device (was %dx%d)"), MaxWidth, MaxHeight, OldMaxWidth, OldMaxHeight);
-					}
-				}
-			}
-
-			// 0 means to use native size
-			int32 Width, Height;
-			if (RequestedContentScaleFactor == 0.0f)
-			{
-				Width = MaxWidth;
-				Height = MaxHeight;
-				UE_LOG(LogAndroid, Log, TEXT("Setting Width=%d and Height=%d (requested scale = 0 = auto)"), Width, Height);
-			}
-			else
-			{
-				if (GAndroidIsPortrait)
-				{
-					Height = 1280 * RequestedContentScaleFactor;
-				}
-				else
-				{
-					Height = 720 * RequestedContentScaleFactor;
-				}
-
-				// apply the aspect ration to get the width
-				Width = Height * AspectRatio;
-
-				// clamp to native resolution
-				Width = FPlatformMath::Min(Width, MaxWidth);
-				Height = FPlatformMath::Min(Height, MaxHeight);
-
-				UE_LOG(LogAndroid, Log, TEXT("Setting Width=%d and Height=%d (requested scale = %f)"), Width, Height, RequestedContentScaleFactor);
-			}
-
-			ScreenWidth = Width;
-			ScreenHeight = Height;
-			break;
-		}
-		case EPIEPreviewDeviceType::IOS:
-		{
-			bool GIOSIsPortrait = false;
-
-			// some phones gave it the other way (so, if swap if the app is landscape, but width < height)
-			if ((GIOSIsPortrait && ScreenWidth > ScreenHeight) ||
-				(!GIOSIsPortrait && ScreenWidth < ScreenHeight))
-			{
-				Swap(ScreenWidth, ScreenHeight);
-			}
-			break;
-		}
+		Device->ApplyRHIPrerequisitesOverrides();
+		DeviceProfile = Device->GetProfile();
 	}
 }
 
-
 TSharedRef<SWindow> FPIEPreviewDeviceModule::CreatePIEPreviewDeviceWindow(FVector2D ClientSize, FText WindowTitle, EAutoCenter AutoCenterType, FVector2D ScreenPosition, TOptional<float> MaxWindowWidth, TOptional<float> MaxWindowHeight)
 {
+	FPIEPreviewWindowCoreStyle::InitializePIECoreStyle();
+
 	static FWindowStyle BackgroundlessStyle = FCoreStyle::Get().GetWidgetStyle<FWindowStyle>("Window");
 	BackgroundlessStyle.SetBackgroundBrush(FSlateNoResource());
-
-	return SNew(SPIEPreviewWindow)
+	TSharedRef<SPIEPreviewWindow> Window = SNew(SPIEPreviewWindow, Device)
 		.Type(EWindowType::GameWindow)
 		.Style(&BackgroundlessStyle)
 		.ClientSize(ClientSize)
@@ -268,89 +77,60 @@ TSharedRef<SWindow> FPIEPreviewDeviceModule::CreatePIEPreviewDeviceWindow(FVecto
 		.SizingRule(ESizingRule::FixedSize)
 		.HasCloseButton(true)
 		.SupportsMinimize(true)
-		.SupportsMaximize(true);
+		.SupportsMaximize(false)
+		.bManualManageDPI(false);
+
+ 	WindowWPtr = Window;
+
+	return Window;
 }
 
-void FPIEPreviewDeviceModule::ApplyPreviewDeviceState()
+void FPIEPreviewDeviceModule::UpdateDisplayResolution()
 {
-	if (!DeviceSpecs.IsValid())
+	if (!Device.IsValid())
 	{
 		return;
 	}
-	FPIEPreviewWindowCoreStyle::InitializePIECoreStyle();
-	EShaderPlatform PreviewPlatform = SP_NumPlatforms;
-	ERHIFeatureLevel::Type PreviewFeatureLevel = GetPreviewDeviceFeatureLevel();
-	FPIERHIOverrideState* RHIOverrideState = nullptr;
-	switch (DeviceSpecs->DevicePlatform)
+
+	const int32 WindowWidth = Device->GetWindowWidth();
+	const int32 WindowHeight = Device->GetWindowHeight();
+
+	FSystemResolution::RequestResolutionChange(WindowWidth, WindowHeight, EWindowMode::Windowed);
+	IConsoleManager::Get().CallAllConsoleVariableSinks();
+}
+
+void FPIEPreviewDeviceModule::PrepareDeviceDisplay()
+{
+	TSharedPtr<SPIEPreviewWindow> WindowPtr = WindowWPtr.Pin();
+
+	if (!WindowPtr.IsValid() && !Device.IsValid())
 	{
-		case EPIEPreviewDeviceType::Android:
-		{
-			if (PreviewFeatureLevel == ERHIFeatureLevel::ES2)
-			{
-				PreviewPlatform = SP_OPENGL_ES2_ANDROID;
-				RHIOverrideState = &DeviceSpecs->AndroidProperties.GLES2RHIState;
-			}
-			else
-			{
-				PreviewPlatform = SP_OPENGL_ES3_1_ANDROID;
-				RHIOverrideState = &DeviceSpecs->AndroidProperties.GLES31RHIState;
-			}
-		}
-		break;
-		case EPIEPreviewDeviceType::IOS:
-		{
-			if (PreviewFeatureLevel == ERHIFeatureLevel::ES2)
-			{
-				PreviewPlatform = SP_OPENGL_ES2_IOS;
-				RHIOverrideState = &DeviceSpecs->IOSProperties.GLES2RHIState;
-			}
-			else
-			{
-				PreviewPlatform = SP_METAL_MACES3_1;
-				RHIOverrideState = &DeviceSpecs->IOSProperties.MetalRHIState;
-			}
-		}
-		break;
+		return;
 	}
 
-	if(PreviewPlatform != SP_NumPlatforms)
+	FSlateApplication::Get().SetGameIsFakingTouchEvents(true);
+
+	WindowPtr->SetScaleWindowToDeviceSize(true);
+
+	UpdateDisplayResolution();
+}
+
+void FPIEPreviewDeviceModule::ApplyPreviewDeviceState()
+{
+	if (!Device.IsValid())
 	{
-		UMaterialShaderQualitySettings* MaterialShaderQualitySettings = UMaterialShaderQualitySettings::Get();
-		FName QualityPreviewShaderPlatform = LegacyShaderPlatformToShaderFormat(PreviewPlatform);
-		MaterialShaderQualitySettings->GetShaderPlatformQualitySettings(QualityPreviewShaderPlatform);
-		MaterialShaderQualitySettings->SetPreviewPlatform(QualityPreviewShaderPlatform);
+		return;
 	}
 
-	ApplyRHIOverrides(RHIOverrideState);
-
 	// TODO: Localization
 	FString AppTitle = FGlobalTabmanager::Get()->GetApplicationTitle().ToString() + "Previewing: "+ PreviewDevice;
 	FGlobalTabmanager::Get()->SetApplicationTitle(FText::FromString(AppTitle));
-	
-	int32 ScreenWidth, ScreenHeight;
-	GetPreviewDeviceResolution(ScreenWidth, ScreenHeight);
-	EWindowMode::Type InWindowMode = EWindowMode::Windowed;
-	FString WindowModeSuffix;
-	switch (InWindowMode)
-	{
-	case EWindowMode::Windowed:
-	{
-		WindowModeSuffix = TEXT("w");
-	} break;
-	case EWindowMode::WindowedFullscreen:
-	{
-		WindowModeSuffix = TEXT("wf");
-	} break;
-	case EWindowMode::Fullscreen:
-	{
-		WindowModeSuffix = TEXT("f");
-	} break;
-	}
-	
-	FString NewValue = FString::Printf(TEXT("%dx%d%s"), ScreenWidth, ScreenHeight, *WindowModeSuffix);
-	static IConsoleVariable* CVarSystemRes = IConsoleManager::Get().FindConsoleVariable(TEXT("r.SetRes"));
-	CVarSystemRes->Set(*NewValue, ECVF_SetByConsole);
-	IConsoleManager::Get().CallAllConsoleVariableSinks();
+
+	int32 TitleBarSize = SPIEPreviewWindow::GetDefaultTitleBarSize();
+	Device->SetupDevice(TitleBarSize);
+
+	// need to call this before the actual window is created in order override the window mode to EWindowMode::Windowed
+	UpdateDisplayResolution();
 }
 
 const FPIEPreviewDeviceContainer& FPIEPreviewDeviceModule::GetPreviewDeviceContainer()
@@ -382,11 +162,12 @@ FString FPIEPreviewDeviceModule::FindDeviceSpecificationFilePath(const FString&
 		}
 	}
 	return FoundPath;
+
 }
 
 bool FPIEPreviewDeviceModule::ReadDeviceSpecification()
 {
-	DeviceSpecs = nullptr;
+	Device = nullptr;
 
 	if (!FParse::Value(FCommandLine::Get(), GetPreviewDeviceCommandSwitch(), PreviewDevice))
 	{
@@ -405,70 +186,19 @@ bool FPIEPreviewDeviceModule::ReadDeviceSpecification()
 			// We need to initialize FPIEPreviewDeviceSpecifications early as device profiles need to be evaluated before ProcessNewlyLoadedUObjects can be called.
 			CreatePackage(nullptr, TEXT("/Script/PIEPreviewDeviceProfileSelector"));
 
-			DeviceSpecs = MakeShareable(new FPIEPreviewDeviceSpecifications());
+			Device = MakeShareable(new FPIEPreviewDevice());
 
-			if (!FJsonObjectConverter::JsonAttributesToUStruct(RootObject->Values, FPIEPreviewDeviceSpecifications::StaticStruct(), DeviceSpecs.Get(), 0, 0))
+			if (!FJsonObjectConverter::JsonAttributesToUStruct(RootObject->Values, FPIEPreviewDeviceSpecifications::StaticStruct(), Device->GetDeviceSpecs().Get(), 0, 0))
 			{
-				DeviceSpecs = nullptr;
+				Device = nullptr;
 			}
 		}
 	}
-	bool bValidDeviceSpec = DeviceSpecs.IsValid();
+	bool bValidDeviceSpec = Device.IsValid();
 	if (!bValidDeviceSpec)
 	{
 		UE_LOG(LogPIEPreviewDevice, Warning, TEXT("Could not load device specifications for preview target device '%s'"), *PreviewDevice);
 	}
-	return bValidDeviceSpec;
-}
-
-ERHIFeatureLevel::Type FPIEPreviewDeviceModule::GetPreviewDeviceFeatureLevel() const
-{
-	check(DeviceSpecs.IsValid());
-
-	switch (DeviceSpecs->DevicePlatform)
-	{
-		case EPIEPreviewDeviceType::Android:
-		{
-			FString SubVersion;
-			// Check for ES3.1+ support from GLVersion, TODO: check other ES31 feature level constraints, see android's PlatformInitOpenGL
-			const bool bDeviceSupportsES31 = DeviceSpecs->AndroidProperties.GLVersion.Split(TEXT("OpenGL ES 3."), nullptr, &SubVersion) && FCString::Atoi(*SubVersion) >= 1;
-
-			// check the project's gles support:
-			bool bProjectBuiltForES2 = false, bProjectBuiltForES31 = false;
-			GConfig->GetBool(TEXT("/Script/AndroidRuntimeSettings.AndroidRuntimeSettings"), TEXT("bBuildForES31"), bProjectBuiltForES31, GEngineIni);
-			GConfig->GetBool(TEXT("/Script/AndroidRuntimeSettings.AndroidRuntimeSettings"), TEXT("bBuildForES2"), bProjectBuiltForES2, GEngineIni);
 
-			// Android Preview Device is currently expected to work on gles.
-			check(bProjectBuiltForES2 || bProjectBuiltForES31);
-
-			// Projects without ES2 support can only expect to run on ES31 devices.
-			check(bProjectBuiltForES2 || bDeviceSupportsES31);
-
-			// ES3.1+ devices fallback to ES2 if the project itself doesn't support ES3.1
-			return bDeviceSupportsES31 && bProjectBuiltForES31 ? ERHIFeatureLevel::ES3_1 : ERHIFeatureLevel::ES2;
-		}
-		case EPIEPreviewDeviceType::IOS:
-		{
-			bool bProjectBuiltForES2 = false, bProjectBuiltForMetal = false, bProjectBuiltForMRTMetal = false;
-			GConfig->GetBool(TEXT("/Script/IOSRuntimeSettings.IOSRuntimeSettings"), TEXT("bSupportsMetal"), bProjectBuiltForMetal, GEngineIni);
-			GConfig->GetBool(TEXT("/Script/IOSRuntimeSettings.IOSRuntimeSettings"), TEXT("bSupportsOpenGLES2"), bProjectBuiltForES2, GEngineIni);
-			GConfig->GetBool(TEXT("/Script/IOSRuntimeSettings.IOSRuntimeSettings"), TEXT("bSupportsMetalMRT"), bProjectBuiltForMRTMetal, GEngineIni);
-			
-			const bool bDeviceSupportsMetal = DeviceSpecs->IOSProperties.MetalRHIState.MaxTextureDimensions > 0;
-
-			// not supporting preview for MRT metal 
-			check(!bProjectBuiltForMRTMetal);
-
-			// atleast one of these should be valid!
-			check(bProjectBuiltForES2 || bProjectBuiltForMetal);
-
-			// if device doesnt support metal the project must have ES2 enabled.
-			check(bProjectBuiltForES2 || (bProjectBuiltForMetal && bDeviceSupportsMetal));
-
-			return (bDeviceSupportsMetal && bProjectBuiltForMetal) ? ERHIFeatureLevel::ES3_1 : ERHIFeatureLevel::ES2;
-		}
-	}
-
-	checkNoEntry();
-	return  ERHIFeatureLevel::Num;
+	return bValidDeviceSpec;
 }
diff --git a/Engine/Source/Editor/PIEPreviewDeviceProfileSelector/Private/PIEPreviewDeviceSpecification.h b/Engine/Source/Editor/PIEPreviewDeviceProfileSelector/Private/PIEPreviewDeviceSpecification.h
--- a/Engine/Source/Editor/PIEPreviewDeviceProfileSelector/Private/PIEPreviewDeviceSpecification.h
+++ b/Engine/Source/Editor/PIEPreviewDeviceProfileSelector/Private/PIEPreviewDeviceSpecification.h
@@ -4,6 +4,7 @@
 
 #include "CoreMinimal.h"
 #include "UObject/Object.h"
+
 #include "PIEPreviewDeviceSpecification.generated.h"
 
 UENUM()
@@ -105,6 +106,9 @@ public:
 	UPROPERTY()
 	FString DeviceModel;
 
+	UPROPERTY()
+	float NativeScaleFactor = 0.0f;
+
 	UPROPERTY()
 	FPIERHIOverrideState GLES2RHIState;
 
@@ -119,15 +123,26 @@ public:
 	GENERATED_USTRUCT_BODY()
 	
 	UPROPERTY()
-	int32 X;
+	int32 X = 0;
 	UPROPERTY()
-	int32 Y;
+	int32 Y = 0;
 	UPROPERTY()
-	int32 Width;
+	int32 Width = 0;
 	UPROPERTY()
-	int32 Height;
+	int32 Height = 0;
 }; 
 
+USTRUCT()
+struct FPIEBezelProperties
+{
+	GENERATED_USTRUCT_BODY()
+
+	UPROPERTY()
+	FString DeviceBezelFile;
+	UPROPERTY()
+	FPIEPreviewDeviceBezelViewportRect BezelViewportRect;
+};
+
 USTRUCT()
 struct FPIEPreviewDeviceSpecifications
 {
@@ -141,13 +156,19 @@ struct FPIEPreviewDeviceSpecifications
 	int32 ResolutionY;
 	UPROPERTY()
 	int32 ResolutionYImmersiveMode;
+
+	UPROPERTY()
+	int32 PPI;
+
+	UPROPERTY()
+	TArray<float> ScaleFactors;
 	
 	UPROPERTY()
-	FPIEAndroidDeviceProperties AndroidProperties;
+	FPIEBezelProperties BezelProperties;
 
 	UPROPERTY()
-	FPIEIOSDeviceProperties IOSProperties;
+	FPIEAndroidDeviceProperties AndroidProperties;
 
 	UPROPERTY()
-	FPIEPreviewDeviceBezelViewportRect BezelViewportRect;
+	FPIEIOSDeviceProperties IOSProperties;
 };
diff --git a/Engine/Source/Editor/PIEPreviewDeviceProfileSelector/Private/PIEPreviewWindow.cpp b/Engine/Source/Editor/PIEPreviewDeviceProfileSelector/Private/PIEPreviewWindow.cpp
--- a/Engine/Source/Editor/PIEPreviewDeviceProfileSelector/Private/PIEPreviewWindow.cpp
+++ b/Engine/Source/Editor/PIEPreviewDeviceProfileSelector/Private/PIEPreviewWindow.cpp
@@ -2,9 +2,122 @@
 
 #include "PIEPreviewWindow.h"
 
+#include "HAL/IConsoleManager.h"
+#include "ImageUtils.h"
+#include "Engine/Texture2D.h"
+#include "HAL/PlatformApplicationMisc.h"
+#include "Widgets/SWindow.h"
+#include "Framework/Application/SlateApplication.h"
+#include "EditorStyleSet.h"
+#include "SEditorViewportToolBarMenu.h"
+#include "Widgets/Input/SMenuAnchor.h"
+#include "Framework/MultiBox/MultiBoxBuilder.h"
+#include "Widgets/Input/SCheckBox.h"
+
+
 #if WITH_EDITOR
-SPIEPreviewWindow::SPIEPreviewWindow() 
+
+//***********************************************************************************
+//SPIEToolbar Implementation
+//***********************************************************************************
+
+/** Toolbar class used to add some menus to configure various device display settings */
+class SPIEToolbar : public SViewportToolBar
+{
+public:
+	SLATE_BEGIN_ARGS(SPIEToolbar) {}
+		SLATE_EVENT(FOnGetContent, OnGetMenuContent)
+	SLATE_END_ARGS()
+
+	void Construct(const FArguments& InArgs);
+	FReply OnMenuClicked();
+
+protected:
+	TSharedPtr<SMenuAnchor> MenuAnchor;
+};
+
+void SPIEToolbar::Construct(const FArguments& InArgs)
+{
+	SViewportToolBar::Construct(SViewportToolBar::FArguments());
+
+	const FSlateBrush* ImageBrush = FPIEPreviewWindowCoreStyle::Get().GetBrush("ComboButton.Arrow");
+
+	const float MenuIconSize = 16.0f;
+
+	TSharedPtr<SWidget> ButtonContent =
+		SNew(SHorizontalBox)
+		+ SHorizontalBox::Slot()
+		.AutoWidth()
+		.Padding(2.0f, 2.0f)
+		[
+			SNew(SHorizontalBox)
+			+ SHorizontalBox::Slot()
+			.AutoWidth()
+			[
+				SNew(SBox)
+				.WidthOverride(MenuIconSize)
+				.HeightOverride(MenuIconSize)
+				.VAlign(VAlign_Center)
+				.HAlign(HAlign_Center)
+				[
+					SNew(SImage)
+					.Image(ImageBrush)
+					.ColorAndOpacity(FSlateColor::UseForeground())
+				]
+			]
+		];
+
+	ChildSlot
+	[
+		SAssignNew(MenuAnchor, SMenuAnchor)
+		.Placement(MenuPlacement_BelowAnchor)
+		[
+			SNew(SButton)
+			// Allows users to drag with the mouse to select options after opening the menu */
+			.ClickMethod(EButtonClickMethod::MouseDown)
+			.ContentPadding(FMargin(2.0f, 2.0f))
+			.VAlign(VAlign_Center)
+			.ButtonStyle(FPIEPreviewWindowCoreStyle::Get(), "PIEWindow.MenuButton")
+			.OnClicked(this, &SPIEToolbar::OnMenuClicked)
+			[
+				ButtonContent.ToSharedRef()
+			]
+		]
+		.OnGetMenuContent(InArgs._OnGetMenuContent)
+	];
+}
+
+FReply SPIEToolbar::OnMenuClicked()
+{
+	// If the menu button is clicked toggle the state of the menu anchor which will open or close the menu
+	if (MenuAnchor->ShouldOpenDueToClick())
+	{
+		MenuAnchor->SetIsOpen(true);
+		SetOpenMenu(MenuAnchor);
+	}
+	else
+	{
+		MenuAnchor->SetIsOpen(false);
+		TSharedPtr<SMenuAnchor> NullAnchor;
+		SetOpenMenu(NullAnchor);
+	}
+
+	return FReply::Handled();
+}
+
+//***********************************************************************************
+//SPIEPreviewWindow Implementation
+//***********************************************************************************
+SPIEPreviewWindow::SPIEPreviewWindow()
 {
+
+}
+
+void SPIEPreviewWindow::Construct(const FArguments& InArgs, TSharedPtr<FPIEPreviewDevice> InDevice)
+{
+	SWindow::Construct(InArgs);
+
+	SetDevice(InDevice);
 }
 
 TSharedRef<SWidget> SPIEPreviewWindow::MakeWindowTitleBar(const TSharedRef<SWindow>& Window, const TSharedPtr<SWidget>& CenterContent, EHorizontalAlignment CenterContentAlignment)
@@ -19,4 +132,437 @@ EHorizontalAlignment SPIEPreviewWindow::GetTitleAlignment()
 	return EHorizontalAlignment::HAlign_Left;
 }
 
+void SPIEPreviewWindow::SetSceneViewportPadding(const FMargin& Margin)
+{
+	ContentSlot->SlotPadding = Margin;
+}
+
+void SPIEPreviewWindow::ComputeBezelOrientation()
+{
+	if (BezelImage.IsValid())
+	{
+		FVector2D WindowSize = GetClientSizeInScreen();
+
+		bool bBezelRotated = Device->IsDeviceFlipped();
+
+		float ScaleX = bBezelRotated ? WindowSize.X / WindowSize.Y : 1.0f;
+		float ScaleY = bBezelRotated ? Inverse(ScaleX) : 1.0f;
+
+		FScale2D Scale = FScale2D(ScaleX, ScaleY);
+		FQuat2D Rotation = FQuat2D(bBezelRotated ? -PI / 2 : 0);
+		FMatrix2x2 ImageTransformationMatrix = Concatenate(Rotation, Scale);
+
+		BezelImage->SetRenderTransform(FSlateRenderTransform(ImageTransformationMatrix));
+	}
+}
+
+void SPIEPreviewWindow::CreateMenuToolBar()
+{
+	AddOverlaySlot()
+		.Padding(5.0f, 3.0f + SWindowDefs::DefaultTitleBarSize, 0, 0)
+		.HAlign(HAlign_Left)
+		.VAlign(VAlign_Top)
+		[
+			SNew(SPIEToolbar)
+			.OnGetMenuContent(this, &SPIEPreviewWindow::BuildSettingsMenu)
+		];
+}
+
+TSharedRef<SWidget> SPIEPreviewWindow::BuildSettingsMenu()
+{
+	FMenuBuilder MenuBuilder(true, nullptr);
+
+	auto ScaleDescriptionWidget = SNew(STextBlock)
+		.Text(FText::FromString(TEXT("Window Scale")))
+		.Justification(ETextJustify::Center);
+
+	MenuBuilder.AddWidget(ScaleDescriptionWidget, FText());
+	MenuBuilder.AddMenuSeparator();
+
+	// create a scaling checkboxes for each scaling factor needed by the emulated device
+	TArray<float>& ArrScaleFactors = Device->GetDeviceSpecs()->ScaleFactors;
+	for (int32 i = 0; i < ArrScaleFactors.Num(); ++i)
+	{
+		const float ScaleFactor = ArrScaleFactors[i];
+
+		FText EntryText = FText::FromString(FString::SanitizeFloat(ArrScaleFactors[i]) + TEXT("x"));
+		auto IsCheckedFunction = [this, ScaleFactor]()
+		{
+			float WindowScaleFactor = GetWindowScaleFactor();
+
+			return FMath::IsNearlyEqual(ScaleFactor, WindowScaleFactor) ? ECheckBoxState::Checked : ECheckBoxState::Unchecked;
+		};
+
+		auto ExecuteActionFunction = [this, ScaleFactor]()
+		{
+			SetScaleWindowToDeviceSize(false);
+			ScaleWindow(ScaleFactor);
+		};
+
+		CreateMenuEntry(MenuBuilder, MoveTemp(EntryText), MoveTemp(IsCheckedFunction), MoveTemp(ExecuteActionFunction));
+	}
+
+	// scale to device size checkbox
+	if (Device->GetDeviceSpecs()->PPI != 0)
+	{
+		FText EntryText = FText::FromString(TEXT("Scale to device size"));
+		auto IsCheckedFunction = [this]()
+		{
+			bool bScaleToDeviceSize = IsManualManageDPIChanges();
+			return bScaleToDeviceSize ? ECheckBoxState::Checked : ECheckBoxState::Unchecked;
+		};
+
+		auto ExecuteActionFunction = [this]()
+		{
+			SetScaleWindowToDeviceSize(true);
+		};
+
+		CreateMenuEntry(MenuBuilder, MoveTemp(EntryText), MoveTemp(IsCheckedFunction), MoveTemp(ExecuteActionFunction));
+
+		MenuBuilder.AddMenuSeparator();
+	}
+
+	// add bClampWindowSizeState checkbox
+	{
+		FText EntryText = FText::FromString(TEXT("Restrict to desktop size"));
+		auto IsCheckedFunction = [this]()
+		{
+			bool bChecked = IsClampingWindowSize();
+			return bChecked ? ECheckBoxState::Checked : ECheckBoxState::Unchecked;
+		};
+
+		auto ExecuteActionFunction = [this]()
+		{
+			bool bClamp = IsClampingWindowSize();
+			SetClampWindowSize(!bClamp);
+		};
+
+		CreateMenuEntry(MenuBuilder, MoveTemp(EntryText), MoveTemp(IsCheckedFunction), MoveTemp(ExecuteActionFunction));
+	}
+	// add checkbox to handle bezel visibility
+	{
+		FText EntryText = FText::FromString(TEXT("Show phone bezel"));
+		auto IsCheckedFunction = [this]()
+		{
+			bool bBezelVisibility = GetBezelVisibility();
+			return bBezelVisibility ? ECheckBoxState::Checked : ECheckBoxState::Unchecked;
+		};
+
+		auto ExecuteActionFunction = [this]()
+		{
+			FlipBezelVisibility();
+		};
+
+		CreateMenuEntry(MenuBuilder, MoveTemp(EntryText), MoveTemp(IsCheckedFunction), MoveTemp(ExecuteActionFunction));
+	}
+
+	return MenuBuilder.MakeWidget();
+}
+
+void SPIEPreviewWindow::CreateMenuEntry(FMenuBuilder& MenuBuilder, FText&& TextEntry, TFunction<ECheckBoxState()>&& IsCheckedFunction, TFunction<void()>&& ExecuteActionFunction)
+{
+	auto Box =
+		SNew(SHorizontalBox)
+		+ SHorizontalBox::Slot()
+		.AutoWidth()
+		.HAlign(EHorizontalAlignment::HAlign_Left)
+		[
+			SNew(STextBlock)
+			.Visibility(EVisibility::HitTestInvisible)
+			.Text(MoveTemp(TextEntry))
+		]
+		+ SHorizontalBox::Slot()
+		.Padding(FMargin(10.0f, 0.0f))
+		.HAlign(EHorizontalAlignment::HAlign_Right)
+		[
+			SNew(SCheckBox)
+			.IsFocusable(false)
+			.IsEnabled(false)
+			.IsChecked_Lambda(MoveTemp(IsCheckedFunction))
+		];
+
+	MenuBuilder.AddMenuEntry(FUIAction(FExecuteAction::CreateLambda(MoveTemp(ExecuteActionFunction))), Box);
+}
+
+void SPIEPreviewWindow::CreatePIEPreviewBezelOverlay(UTexture2D* pBezelImage)
+{
+	if (pBezelImage != nullptr)
+	{
+		BezelBrush.SetResourceObject(pBezelImage);
+		BezelBrush.ImageSize = FVector2D(pBezelImage->GetSizeX(), pBezelImage->GetSizeY());
+
+		auto GetBezelVisibility = [this]()
+		{
+			EVisibility Visibility = EVisibility::Collapsed;
+
+			if (Device.IsValid())
+			{
+				bool bVisible = Device->GetBezelVisibility();
+				Visibility = bVisible ? EVisibility::SelfHitTestInvisible : EVisibility::Collapsed;
+			}
+
+			return Visibility;
+		};
+
+		BezelImage = SNew(SImage)
+			.Image(&BezelBrush)
+			.Visibility_Lambda(GetBezelVisibility)
+			.RenderTransformPivot(FVector2D(.5f, .5f));
+
+		AddOverlaySlot()
+			.Padding(0, SWindowDefs::DefaultTitleBarSize, 0, 0)
+			.HAlign(HAlign_Fill)
+			.VAlign(VAlign_Fill)
+			[
+				BezelImage.ToSharedRef()
+			];
+
+		ComputeBezelOrientation();
+	}
+}
+
+void SPIEPreviewWindow::ScaleWindow(const float ScreenFactor, const float DPIScaleFactor /*= 1.0f*/)
+{
+	if (!Device.IsValid())
+	{
+		return;
+	}
+
+	if (FMath::IsNearlyEqual(ScreenFactor, CachedScaleToDeviceFactor) && FMath::IsNearlyEqual(DPIScaleFactor, CachedDPIScaleFactor))
+	{
+		return;
+	}
+
+	CachedScaleToDeviceFactor = ScreenFactor;
+	CachedDPIScaleFactor = DPIScaleFactor;
+
+	if (IsManualManageDPIChanges())
+	{
+		FSlateApplication::Get().HandleDPIScaleChanged(GetNativeWindow().ToSharedRef());
+	}
+
+	Device->ScaleResolution(ScreenFactor, DPIScaleFactor, bClampWindowSizeState);
+
+	UpdateWindow();
+}
+
+float SPIEPreviewWindow::GetWindowScaleFactor() const
+{
+	float ScaleFactor = 1.0f;
+
+	if (Device.IsValid())
+	{
+		ScaleFactor = Device->GetResolutionScale();
+	}
+
+	return ScaleFactor;
+}
+
+void SPIEPreviewWindow::OnWindowMoved(const TSharedRef<SWindow>& Window)
+{
+	if (IsManualManageDPIChanges())
+	{
+		ScaleDeviceToPhisicalSize();
+	}
+}
+
+void SPIEPreviewWindow::OnDisplayDPIChanged(TSharedRef<SWindow> Window)
+{
+	if (IsManualManageDPIChanges())
+	{
+		ScaleDeviceToPhisicalSize();
+	}
+}
+
+void SPIEPreviewWindow::SetScaleWindowToDeviceSize(const bool bScale)
+{
+	if (bScale != IsManualManageDPIChanges())
+	{
+		SetManualManageDPIChanges(bScale);
+
+		if (Device.IsValid())
+		{
+			Device->SetIgnoreMobileContentScaleFactor(bScale);
+		}
+
+		if (IsManualManageDPIChanges())
+		{
+			ScaleDeviceToPhisicalSize();
+
+			SetOnWindowMoved(FOnWindowMoved::CreateSP(this, &SPIEPreviewWindow::OnWindowMoved));
+			HandleDPIChange = FSlateApplication::Get().OnSystemSignalsDPIChanged().AddSP(this, &SPIEPreviewWindow::OnDisplayDPIChanged);
+		}
+		else
+		{
+			SetOnWindowMoved(nullptr);
+
+			if (HandleDPIChange.IsValid())
+			{
+				FSlateApplication::Get().OnSystemSignalsDPIChanged().Remove(HandleDPIChange);
+			}
+		}
+	}
+}
+
+void SPIEPreviewWindow::ScaleDeviceToPhisicalSize()
+{
+	if (Device.IsValid())
+	{
+		float ScreenFactor;
+		float DPIScaleFactor;
+		ComputeScaleToDeviceSizeFactor(ScreenFactor, DPIScaleFactor);
+
+		SetDPIScaleFactor(DPIScaleFactor);
+		ScaleWindow(ScreenFactor, DPIScaleFactor);
+	}
+}
+
+void SPIEPreviewWindow::ComputeScaleToDeviceSizeFactor(float& OutScreenFactor, float& OutDPIScaleFactor) const
+{
+	OutScreenFactor = 1.0f;
+	OutDPIScaleFactor = 1.0f;
+
+	FDisplayMetrics DisplayMetrics;
+	FDisplayMetrics::GetDisplayMetrics(DisplayMetrics);
+
+	FVector2D WindowPos = GetPositionInScreen();
+	int32 PointX = FMath::RoundToInt(WindowPos.X);
+	int32 PointY = FMath::RoundToInt(WindowPos.Y);
+
+	FPlatformRect& VirtualDisplayRect = DisplayMetrics.VirtualDisplayRect;
+	PointX = FMath::Clamp(PointX, VirtualDisplayRect.Left, VirtualDisplayRect.Right);
+	PointY = FMath::Clamp(PointY, VirtualDisplayRect.Top, VirtualDisplayRect.Bottom);
+
+	float RatioMonitorResolution = 1.0f;
+	int32 LocalPPI = 0;
+
+	TArray<FMonitorInfo>& MonitorInfoArray = DisplayMetrics.MonitorInfo;
+	for (int32 i = 0; i < MonitorInfoArray.Num(); ++i)
+	{
+		FMonitorInfo& MonitorInfo = MonitorInfoArray[i];
+
+		const int32 PointOffset = 0;
+		if (PointX >= MonitorInfo.DisplayRect.Left &&
+			PointX <= MonitorInfo.DisplayRect.Right &&
+			PointY >= MonitorInfo.DisplayRect.Top &&
+			PointY <= MonitorInfo.DisplayRect.Bottom)
+		{
+			int32 MonitorWidth = MonitorInfo.DisplayRect.Right - MonitorInfo.DisplayRect.Left;
+			int32 MonitorHeight = MonitorInfo.DisplayRect.Bottom - MonitorInfo.DisplayRect.Top;
+
+			float MonitorResolutionScale = FMath::Min((float)MonitorWidth / (float)MonitorInfo.NativeWidth, (float)MonitorHeight / (float)MonitorInfo.NativeHeight);
+
+			float NativeRatio = (float)MonitorInfo.NativeWidth / (float)MonitorInfo.NativeHeight;
+			float CurrentRatio = (float)MonitorWidth / (float)MonitorHeight;
+			float MonitorPixelRatio = NativeRatio / CurrentRatio;
+
+			RatioMonitorResolution = MonitorResolutionScale * MonitorPixelRatio;
+
+			LocalPPI = MonitorInfo.DPI;
+
+			break;
+		}
+	}
+
+	OutDPIScaleFactor = FPlatformApplicationMisc::GetDPIScaleFactorAtPoint(PointX, PointY);
+
+	const int32 DevicePPI = Device->GetDeviceSpecs()->PPI;
+	float PPIRatio = (!!DevicePPI && !!LocalPPI) ? (float)LocalPPI / (float)DevicePPI : 1.0f;
+	OutScreenFactor = PPIRatio * RatioMonitorResolution / OutDPIScaleFactor;
+}
+
+void SPIEPreviewWindow::RotateWindow()
+{
+	if (!Device.IsValid())
+	{
+		return;
+	}
+
+	Device->SwitchOrientation(bClampWindowSizeState);
+
+	UpdateWindow();
+}
+
+void SPIEPreviewWindow::FlipBezelVisibility()
+{
+	if (!Device.IsValid())
+	{
+		return;
+	}
+
+	bool bVisible = Device->GetBezelVisibility();
+	Device->SetBezelVisibility(!bVisible, bClampWindowSizeState);
+
+	UpdateWindow();
+}
+
+bool SPIEPreviewWindow::GetBezelVisibility() const
+{
+	bool bVisible = false;
+
+	if (Device.IsValid())
+	{
+		bVisible = Device->GetBezelVisibility();
+	}
+
+	return bVisible;
+}
+
+void SPIEPreviewWindow::UpdateWindow()
+{
+	if (!Device.IsValid())
+	{
+		return;
+	}
+
+	// compute position window position
+	// try to maintain its old top left corner position, but while keeping it inside the desktop area
+	FVector2D WindowPos = GetPositionInScreen();
+	int32 PosX = (int32)WindowPos.X;
+	int32 PosY = (int32)WindowPos.Y;
+
+	FDisplayMetrics DisplayMetrics;
+	FDisplayMetrics::GetDisplayMetrics(DisplayMetrics);
+
+	if (PosX + Device->GetWindowWidth() > DisplayMetrics.VirtualDisplayRect.Right)
+	{
+		PosX = DisplayMetrics.VirtualDisplayRect.Right - Device->GetWindowWidth();
+	}
+	PosX = FMath::Max(DisplayMetrics.VirtualDisplayRect.Left, PosX);
+
+	if (PosY + Device->GetWindowHeight() > DisplayMetrics.VirtualDisplayRect.Bottom)
+	{
+		PosY = DisplayMetrics.VirtualDisplayRect.Bottom - Device->GetWindowHeight() - SWindowDefs::DefaultTitleBarSize;
+	}
+	PosY = FMath::Max(DisplayMetrics.VirtualDisplayRect.Top, PosY);
+
+	ReshapeWindow(FVector2D(PosX, PosY), FVector2D(Device->GetWindowWidth(), Device->GetWindowHeight()));
+
+	// offset the viewport widget into its correct location
+	SetSceneViewportPadding(Device->GetViewportMargin());
+
+	// bezel orientation depends on the window size so we need to call it after ReshapeWindow()
+	ComputeBezelOrientation();
+}
+
+void SPIEPreviewWindow::SetDevice(TSharedPtr<FPIEPreviewDevice> InDevice)
+{
+	Device = InDevice;
+
+	if (Device.IsValid())
+	{
+		FMargin ViewportOffset = Device->GetViewportMargin();
+		SetSceneViewportPadding(ViewportOffset);
+
+		CreatePIEPreviewBezelOverlay(InDevice->GetBezelTexture());
+
+		CreateMenuToolBar();
+	}
+}
+
+int32 SPIEPreviewWindow::GetDefaultTitleBarSize()
+{
+	return SWindowDefs::DefaultTitleBarSize;
+}
+
 #endif
\ No newline at end of file
diff --git a/Engine/Source/Editor/PIEPreviewDeviceProfileSelector/Private/PIEPreviewWindow.h b/Engine/Source/Editor/PIEPreviewDeviceProfileSelector/Private/PIEPreviewWindow.h
--- a/Engine/Source/Editor/PIEPreviewDeviceProfileSelector/Private/PIEPreviewWindow.h
+++ b/Engine/Source/Editor/PIEPreviewDeviceProfileSelector/Private/PIEPreviewWindow.h
@@ -6,6 +6,7 @@
 
 #include "Widgets/SWindow.h"
 #include "PIEPreviewWindowTitleBar.h"
+#include "PIEPreviewDevice.h"
 
 class PIEPREVIEWDEVICEPROFILESELECTOR_API SPIEPreviewWindow
 	: public SWindow
@@ -16,8 +17,122 @@ public:
 	*/
 	SPIEPreviewWindow();
 
+	void Construct(const FArguments& InArgs, TSharedPtr<FPIEPreviewDevice> InDevice);
+
+	// call this to 
+	void SetSceneViewportPadding(const FMargin& Margin);
+
+	// sets the device displayed by this window
+	void SetDevice(TSharedPtr<class FPIEPreviewDevice> InDevice);
+
+	// retrieves the display device
+	TSharedPtr<class FPIEPreviewDevice> GetDevice();
+
+	// call to perform resolution scaling on the device
+	void ScaleWindow(const float ScaleFactor, const float BezelFactor = 1.0f);
+
+	// call to retrieve device's current resolution scale factor
+	float GetWindowScaleFactor() const;
+
+	// call to rotate the window emulating a rotation of the physical device
+	void RotateWindow();
+
+	// wrapper function that will query the same functionality on the device
+	bool IsRotationAllowed() const;
+
+	// call to enable/disable bezel visibility
+	void FlipBezelVisibility();
+
+	// call to determine whether or not the phone bezel is displayed 
+	bool GetBezelVisibility() const;
+
+	static int32 GetDefaultTitleBarSize();
+
+	// call with a true argument to resize the window so it matches the device physical size
+	void SetScaleWindowToDeviceSize(const bool bScale);
+
+	// call with a true argument to restrict the window size to the desktop size
+	void SetClampWindowSize(const bool bClamp);
+	bool IsClampingWindowSize() const;
+
 private:
 	virtual EHorizontalAlignment GetTitleAlignment() override;
 	virtual TSharedRef<SWidget> MakeWindowTitleBar(const TSharedRef<SWindow>& Window, const TSharedPtr<SWidget>& CenterContent, EHorizontalAlignment CenterContentAlignment) override;
+
+	// this will properly scale and rotate the bezel to match the orientation of the device
+	void ComputeBezelOrientation();
+
+	// function that adds a widget to display the device's bezel
+	void CreatePIEPreviewBezelOverlay(UTexture2D* pBezelImage);
+
+	// function used create a menu toolbar
+	void CreateMenuToolBar();
+
+	/** creates and returns the settings menu */
+	TSharedRef<SWidget> BuildSettingsMenu();
+
+	/** helper function used to create a menu item formed from a description text and a checkbox */
+	/** @ IsCheckedFunction - computes the state of the checkbox */
+	/** @ ExecuteActionFunction - function to be called when this menu item is clicked */
+	void CreateMenuEntry(class FMenuBuilder& MenuBuilder, FText&& TextEntry, TFunction<ECheckBoxState()>&& IsCheckedFunction, TFunction<void()>&& ExecuteActionFunction);
+
+	// call to compute window size, position on screen and bezel orientation
+	void UpdateWindow();
+
+	// call function to make the size of this window match the physical device size
+	void ScaleDeviceToPhisicalSize();
+
+	// utility function that will compute a screen and DPI scale factor needed to scale the display window to the physical device size
+	void ComputeScaleToDeviceSizeFactor(float& OutScreenFactor, float& OutDPIScaleFactor) const;
+
+	// callback when DPI value varies, useful when we need to constrain the window to the physical device size
+	void OnWindowMoved(const TSharedRef<SWindow>& Window);
+
+	// callback when DPI value varies, useful when we need to constrain the window to the physical device size
+	void OnDisplayDPIChanged(TSharedRef<SWindow> Window);
+
+protected:
+	// brush created to display the bezel
+	FSlateBrush BezelBrush;
+
+	// pointer to the actual bezel image
+	TSharedPtr<SImage> BezelImage;
+
+	// pointer to the device that this window will display
+	TSharedPtr<class FPIEPreviewDevice> Device;
+
+	float CachedScaleToDeviceFactor = 0.0f;
+	float CachedDPIScaleFactor = 0.0f;
+
+	/** when true the window size will be restricted to the desktop size  */
+	bool bClampWindowSizeState = true;
+
+	FDelegateHandle HandleDPIChange;
 };
+
+FORCEINLINE bool SPIEPreviewWindow::IsRotationAllowed() const
+{
+	if (!Device.IsValid())
+	{
+		return false;
+	}
+
+	return Device->IsRotationAllowed();
+}
+
+FORCEINLINE TSharedPtr<class FPIEPreviewDevice> SPIEPreviewWindow::GetDevice()
+{
+	return Device;
+}
+
+FORCEINLINE void SPIEPreviewWindow::SetClampWindowSize(const bool bClamp)
+{
+	bClampWindowSizeState = bClamp;
+}
+
+FORCEINLINE bool SPIEPreviewWindow::IsClampingWindowSize() const
+{
+	return bClampWindowSizeState;
+}
+
 #endif
\ No newline at end of file
diff --git a/Engine/Source/Editor/PIEPreviewDeviceProfileSelector/Private/PIEPreviewWindowCoreStyle.cpp b/Engine/Source/Editor/PIEPreviewDeviceProfileSelector/Private/PIEPreviewWindowCoreStyle.cpp
--- a/Engine/Source/Editor/PIEPreviewDeviceProfileSelector/Private/PIEPreviewWindowCoreStyle.cpp
+++ b/Engine/Source/Editor/PIEPreviewDeviceProfileSelector/Private/PIEPreviewWindowCoreStyle.cpp
@@ -10,53 +10,47 @@
 #include "Brushes/SlateImageBrush.h"
 #include "Styling/SlateStyleRegistry.h"
 #include "Styling/SlateTypes.h"
+#include "PIEPreviewDeviceSpecification.h"
 
 #if WITH_EDITOR
 
 TSharedPtr< ISlateStyle > FPIEPreviewWindowCoreStyle::Instance = nullptr;
 
 #define IMAGE_BRUSH( RelativePath, ... ) FSlateImageBrush(Style->RootToContentDir(RelativePath, TEXT(".png")), __VA_ARGS__)
+#define DEFAULT_FONT(...) FCoreStyle::GetDefaultFontStyle(__VA_ARGS__)
+#define BOX_BRUSH( RelativePath, ... ) FSlateBoxBrush(Style->RootToContentDir( RelativePath, TEXT(".png") ), __VA_ARGS__)
 
 TSharedRef<class ISlateStyle> FPIEPreviewWindowCoreStyle::Create(const FName& InStyleSetName)
 {
-
 	TSharedRef<FSlateStyleSet> Style = MakeShareable(new FSlateStyleSet(InStyleSetName));
 	Style->SetContentRoot(FPaths::EngineContentDir() / TEXT("Slate"));
 
 	const FButtonStyle Button = FCoreStyle::Get().GetWidgetStyle<FButtonStyle>("Button");
 
 	const FButtonStyle ScreenRotationButtonStyle = FButtonStyle(Button)
+		.SetDisabled(IMAGE_BRUSH("Icons/PIEWindow/WindowButton_Screen_Rotation_Disabled", FVector2D(23.0f, 18.0f)))
 		.SetNormal(IMAGE_BRUSH("Icons/PIEWindow/WindowButton_Screen_Rotation_Normal", FVector2D(23.0f, 18.0f)))
 		.SetHovered(IMAGE_BRUSH("Icons/PIEWindow/WindowButton_Screen_Rotation_Hovered", FVector2D(23.0f, 18.0f)))
 		.SetPressed(IMAGE_BRUSH("Icons/PIEWindow/WindowButton_Screen_Rotation_Pressed", FVector2D(23.0f, 18.0f)));
 
-	const FButtonStyle QuarterMobileContentScaleFactorButtonStyle = FButtonStyle(Button)
-		.SetNormal(IMAGE_BRUSH("Icons/PIEWindow/WindowButton_025x_Normal", FVector2D(23.0f, 18.0f)))
-		.SetHovered(IMAGE_BRUSH("Icons/PIEWindow/WindowButton_025x_Hovered", FVector2D(23.0f, 18.0f)))
-		.SetPressed(IMAGE_BRUSH("Icons/PIEWindow/WindowButton_025x_Pressed", FVector2D(23.0f, 18.0f)));
+	Style->Set("PIEWindow", FPIEPreviewWindowStyle()
+		.SetScreenRotationButtonStyle(ScreenRotationButtonStyle));
 
-	const FButtonStyle HalfMobileContentScaleFactorButtonStyle = FButtonStyle(Button)
-		.SetNormal(IMAGE_BRUSH("Icons/PIEWindow/WindowButton_05x_Normal", FVector2D(23.0f, 18.0f)))
-		.SetHovered(IMAGE_BRUSH("Icons/PIEWindow/WindowButton_05x_Hovered", FVector2D(23.0f, 18.0f)))
-		.SetPressed(IMAGE_BRUSH("Icons/PIEWindow/WindowButton_05x_Pressed", FVector2D(23.0f, 18.0f)));
+	Style->Set("PIEWindow.Font", DEFAULT_FONT("Bold", 9));
 
-	const FButtonStyle FullMobileContentScaleFactorButtonStyle = FButtonStyle(Button)
-		.SetNormal(IMAGE_BRUSH("Icons/PIEWindow/WindowButton_1x_Normal", FVector2D(23.0f, 18)))
-		.SetHovered(IMAGE_BRUSH("Icons/PIEWindow/WindowButton_1x_Hovered", FVector2D(23.0f, 18)))
-		.SetPressed(IMAGE_BRUSH("Icons/PIEWindow/WindowButton_1x_Pressed", FVector2D(23.0f, 18)));
+	Style->Set("PIEWindow.MenuButton", FButtonStyle(Button)
+		.SetNormal(BOX_BRUSH("Icons/PIEWindow/SmallRoundedButton", FMargin(7.f / 16.f), FLinearColor(1, 1, 1, 0.75f)))
+		.SetHovered(BOX_BRUSH("Icons/PIEWindow/SmallRoundedButton", FMargin(7.f / 16.f), FLinearColor(1, 1, 1, 1.0f)))
+		.SetPressed(BOX_BRUSH("Icons/PIEWindow/SmallRoundedButton", FMargin(7.f / 16.f))));
 
-	Style->Set("PIEWindow", FPIEPreviewWindowStyle()
-		.SetScreenRotationButtonStyle(ScreenRotationButtonStyle)
-		.SetQuarterMobileContentScaleFactorButtonStyle(QuarterMobileContentScaleFactorButtonStyle)
-		.SetHalfMobileContentScaleFactorButtonStyle(HalfMobileContentScaleFactorButtonStyle)
-		.SetFullMobileContentScaleFactorButtonStyle(FullMobileContentScaleFactorButtonStyle)
-	);
+	Style->Set("ComboButton.Arrow", new IMAGE_BRUSH("Common/ComboArrow", FVector2D(8.0f, 8.0f)));
 
 	return Style;
 }
 
 #undef IMAGE_BRUSH
-
+#undef DEFAULT_FONT
+#undef BOX_BRUSH
 
 void FPIEPreviewWindowCoreStyle::InitializePIECoreStyle()
 {
diff --git a/Engine/Source/Editor/PIEPreviewDeviceProfileSelector/Private/PIEPreviewWindowTitleBar.cpp b/Engine/Source/Editor/PIEPreviewDeviceProfileSelector/Private/PIEPreviewWindowTitleBar.cpp
--- a/Engine/Source/Editor/PIEPreviewDeviceProfileSelector/Private/PIEPreviewWindowTitleBar.cpp
+++ b/Engine/Source/Editor/PIEPreviewDeviceProfileSelector/Private/PIEPreviewWindowTitleBar.cpp
@@ -1,19 +1,37 @@
 // Copyright 1998-2018 Epic Games, Inc. All Rights Reserved.
 
 #include "PIEPreviewWindowTitleBar.h"
+#include "PIEPreviewWindow.h"
+#include "Widgets/Input/SCheckBox.h"
+#include "Framework/MultiBox/MultiBoxBuilder.h"
 
 #if WITH_EDITOR
 
+#define LOCTEXT_NAMESPACE "PIEPreviewWindowTitleBar"
+
 void SPIEPreviewWindowTitleBar::MakeTitleBarContentWidgets(TSharedPtr< SWidget >& OutLeftContent, TSharedPtr< SWidget >& OutRightContent)
 {
 	TSharedPtr< SWidget > OutRightContentBaseWindow;
 	SWindowTitleBar::MakeTitleBarContentWidgets(OutLeftContent, OutRightContentBaseWindow);
-	
+
 	ScreenRotationButton = SNew(SButton)
 		.IsFocusable(false)
 		.IsEnabled(true)
 		.ContentPadding(0)
-		.OnClicked(this, &SPIEPreviewWindowTitleBar::ScreenRotationButton_OnClicked)
+		.OnClicked_Lambda
+		(
+			[this]()
+			{
+				TSharedPtr<SPIEPreviewWindow> PIEWindow = GetOwnerWindow();
+
+				if (PIEWindow.IsValid())
+				{
+					PIEWindow->RotateWindow();
+				}
+
+				return FReply::Handled();
+			}
+		)
 		.Cursor(EMouseCursor::Default)
 		.ButtonStyle(FCoreStyle::Get(), "NoBorder")
 		[
@@ -21,76 +39,32 @@ void SPIEPreviewWindowTitleBar::MakeTitleBarContentWidgets(TSharedPtr< SWidget >
 			.Image(this, &SPIEPreviewWindowTitleBar::GetScreenRotationButtonImage)
 			.ColorAndOpacity(this, &SPIEPreviewWindowTitleBar::GetWindowTitleContentColor)
 		]
+		.IsEnabled_Lambda
+		(
+			[this]()
+			{
+				TSharedPtr<SPIEPreviewWindow> PIEWindow = GetOwnerWindow();
+			
+				if (PIEWindow.IsValid())
+				{
+					return PIEWindow->IsRotationAllowed();
+				}
+
+				return false;
+			}
+		)
 	;
 
-	QuarterMobileContentScaleFactorButton = SNew(SButton)
-		.IsFocusable(false)
-		.IsEnabled(true)
-		.ContentPadding(0)
-		.OnClicked(this, &SPIEPreviewWindowTitleBar::QuarterMobileContentScaleFactorButton_OnClicked)
-		.Cursor(EMouseCursor::Default)
-		.ButtonStyle(FCoreStyle::Get(), "NoBorder")
-		[
-			SNew(SImage)
-			.Image(this, &SPIEPreviewWindowTitleBar::GetQuarterMobileContentScaleFactorImage)
-			.ColorAndOpacity(this, &SPIEPreviewWindowTitleBar::GetWindowTitleContentColor)
-		]
-	;
-
-	HalfMobileContentScaleFactorButton = SNew(SButton)
-		.IsFocusable(false)
-		.IsEnabled(true)
-		.ContentPadding(0.0f)
-		.OnClicked(this, &SPIEPreviewWindowTitleBar::HalfMobileContentScaleFactorButton_OnClicked)
-		.Cursor(EMouseCursor::Default)
-		.ButtonStyle(FCoreStyle::Get(), "NoBorder")
-		[
-			SNew(SImage)
-			.Image(this, &SPIEPreviewWindowTitleBar::GetHalfMobileContentScaleFactorImage)
-			.ColorAndOpacity(this, &SPIEPreviewWindowTitleBar::GetWindowTitleContentColor)
-		]
-	;
-
-	FullMobileContentScaleFactorButton = SNew(SButton)
-		.IsFocusable(false)
-		.IsEnabled(true)
-		.ContentPadding(0.0f)
-		.OnClicked(this, &SPIEPreviewWindowTitleBar::FullMobileContentScaleFactorButton_OnClicked)
-		.Cursor(EMouseCursor::Default)
-		.ButtonStyle(FCoreStyle::Get(), "NoBorder")
-		[
-			SNew(SImage)
-			.Image(this, &SPIEPreviewWindowTitleBar::GetFullMobileContentScaleFactorImage)
-			.ColorAndOpacity(this, &SPIEPreviewWindowTitleBar::GetWindowTitleContentColor)
-		];
-
-
-
 	TSharedRef< SHorizontalBox > WindowTitleBarButtons =
 	SNew(SHorizontalBox)
 	.Visibility(EVisibility::SelfHitTestInvisible);
 
 	WindowTitleBarButtons->AddSlot()
-	.AutoWidth()
-	[
-	ScreenRotationButton.ToSharedRef()
-	];
+		.AutoWidth()
+		[
+			ScreenRotationButton.ToSharedRef()
+		];
 
-	WindowTitleBarButtons->AddSlot()
-	.AutoWidth()
-	[
-	QuarterMobileContentScaleFactorButton.ToSharedRef()
-	];
-	WindowTitleBarButtons->AddSlot()
-	.AutoWidth()
-	[
-	HalfMobileContentScaleFactorButton.ToSharedRef()
-	];
-	WindowTitleBarButtons->AddSlot()
-	.AutoWidth()
-	[
-	FullMobileContentScaleFactorButton.ToSharedRef()
-	];
 	if (OutRightContentBaseWindow.IsValid())
 	{
 		WindowTitleBarButtons->AddSlot()
@@ -99,7 +73,6 @@ void SPIEPreviewWindowTitleBar::MakeTitleBarContentWidgets(TSharedPtr< SWidget >
 				OutRightContentBaseWindow.ToSharedRef()
 			];
 	}
-	
 
 	OutRightContent = SNew(SBox)
 		.Visibility(EVisibility::SelfHitTestInvisible)
@@ -108,31 +81,6 @@ void SPIEPreviewWindowTitleBar::MakeTitleBarContentWidgets(TSharedPtr< SWidget >
 			WindowTitleBarButtons
 		];
 }
-	
-void SPIEPreviewWindowTitleBar::ApplyWindowRotation(TSharedPtr<SWindow> OwnerWindow)
-{
-	FVector2D Size = OwnerWindow->GetInitialDesiredSizeInScreen() * ScaleFactor;
-	IsPortrait = !IsPortrait;
-
-	if (IsPortrait)
-		Swap(Size.X, Size.Y);
-
-	int x, y, w, h;
-	OwnerWindow->GetNativeWindow()->GetFullScreenInfo(x, y, w, h);
-	OwnerWindow->GetNativeWindow()->ReshapeWindow((w - Size.X) * 0.5f, (h - Size.Y) * 0.5f, Size.X, Size.Y);
-}
-
-void SPIEPreviewWindowTitleBar::ApplyWindowScaleFactor(TSharedPtr<SWindow> OwnerWindow)
-{
-	FVector2D Size = OwnerWindow->GetInitialDesiredSizeInScreen() * ScaleFactor;
-
-	if (IsPortrait)
-		Swap(Size.X, Size.Y);
-
-	int x, y, w, h;
-	OwnerWindow->GetNativeWindow()->GetFullScreenInfo(x, y, w, h);
-	OwnerWindow->GetNativeWindow()->ReshapeWindow((w - Size.X) * 0.5f, (h - Size.Y) * 0.5f, Size.X, Size.Y);
-}
 
 const FSlateBrush* SPIEPreviewWindowTitleBar::GetScreenRotationButtonImage() const
 {
@@ -155,130 +103,6 @@ const FSlateBrush* SPIEPreviewWindowTitleBar::GetScreenRotationButtonImage() con
 	{
 		return &FPIEPreviewWindowCoreStyle::Get().GetWidgetStyle<FPIEPreviewWindowStyle>("PIEWindow").ScreenRotationButtonStyle.Normal;
 	}
-
-}
-
-const FSlateBrush* SPIEPreviewWindowTitleBar::GetQuarterMobileContentScaleFactorImage() const
-{
-	TSharedPtr<SWindow> OwnerWindow = OwnerWindowPtr.Pin();
-	
-	if (!OwnerWindow.IsValid())
-	{
-		return nullptr;
-	}
-
-	if (QuarterMobileContentScaleFactorButton->IsPressed())
-	{
-		return &FPIEPreviewWindowCoreStyle::Get().GetWidgetStyle<FPIEPreviewWindowStyle>("PIEWindow").QuarterMobileContentScaleFactorButtonStyle.Pressed;
-	}
-	else if (QuarterMobileContentScaleFactorButton->IsHovered())
-	{
-		return &FPIEPreviewWindowCoreStyle::Get().GetWidgetStyle<FPIEPreviewWindowStyle>("PIEWindow").QuarterMobileContentScaleFactorButtonStyle.Hovered;
-	}
-	else
-	{
-		return &FPIEPreviewWindowCoreStyle::Get().GetWidgetStyle<FPIEPreviewWindowStyle>("PIEWindow").QuarterMobileContentScaleFactorButtonStyle.Normal;
-	}
-
-}
-
-const FSlateBrush* SPIEPreviewWindowTitleBar::GetHalfMobileContentScaleFactorImage() const
-{
-	TSharedPtr<SWindow> OwnerWindow = OwnerWindowPtr.Pin();
-	
-	if (!OwnerWindow.IsValid())
-	{
-		return nullptr;
-	}
-
-	if (HalfMobileContentScaleFactorButton->IsPressed())
-	{
-		return &FPIEPreviewWindowCoreStyle::Get().GetWidgetStyle<FPIEPreviewWindowStyle>("PIEWindow").HalfMobileContentScaleFactorButtonStyle.Pressed;
-	}
-	else if (HalfMobileContentScaleFactorButton->IsHovered())
-	{
-		return &FPIEPreviewWindowCoreStyle::Get().GetWidgetStyle<FPIEPreviewWindowStyle>("PIEWindow").HalfMobileContentScaleFactorButtonStyle.Hovered;
-	}
-	else
-	{
-		return &FPIEPreviewWindowCoreStyle::Get().GetWidgetStyle<FPIEPreviewWindowStyle>("PIEWindow").HalfMobileContentScaleFactorButtonStyle.Normal;
-	}
-
-}
-
-const FSlateBrush* SPIEPreviewWindowTitleBar::GetFullMobileContentScaleFactorImage() const
-{
-	TSharedPtr<SWindow> OwnerWindow = OwnerWindowPtr.Pin();
-	
-	if (!OwnerWindow.IsValid())
-	{
-		return nullptr;
-	}
-
-	if (FullMobileContentScaleFactorButton->IsPressed())
-	{
-		return &FPIEPreviewWindowCoreStyle::Get().GetWidgetStyle<FPIEPreviewWindowStyle>("PIEWindow").FullMobileContentScaleFactorButtonStyle.Pressed;
-	}
-	else if (FullMobileContentScaleFactorButton->IsHovered())
-	{
-		return &FPIEPreviewWindowCoreStyle::Get().GetWidgetStyle<FPIEPreviewWindowStyle>("PIEWindow").FullMobileContentScaleFactorButtonStyle.Hovered;
-	}
-	else
-	{
-		return &FPIEPreviewWindowCoreStyle::Get().GetWidgetStyle<FPIEPreviewWindowStyle>("PIEWindow").FullMobileContentScaleFactorButtonStyle.Normal;
-	}
-
-}
-
-FReply SPIEPreviewWindowTitleBar::ScreenRotationButton_OnClicked()
-{
-	TSharedPtr<SWindow> OwnerWindow = OwnerWindowPtr.Pin();
-
-	if (OwnerWindow.IsValid())
-	{
-		ApplyWindowRotation(OwnerWindow);
-	}
-
-	return FReply::Handled();
-}
-
-FReply SPIEPreviewWindowTitleBar::QuarterMobileContentScaleFactorButton_OnClicked()
-{
-	TSharedPtr<SWindow> OwnerWindow = OwnerWindowPtr.Pin();
-
-	if (OwnerWindow.IsValid())
-	{
-		ScaleFactor = 0.25f;
-		ApplyWindowScaleFactor(OwnerWindow);
-	}
-
-	return FReply::Handled();
-}
-
-FReply SPIEPreviewWindowTitleBar::HalfMobileContentScaleFactorButton_OnClicked()
-{
-	TSharedPtr<SWindow> OwnerWindow = OwnerWindowPtr.Pin();
-
-	if (OwnerWindow.IsValid())
-	{
-		ScaleFactor = 0.5f;
-		ApplyWindowScaleFactor(OwnerWindow);
-	}
-
-	return FReply::Handled();
-}
-
-FReply SPIEPreviewWindowTitleBar::FullMobileContentScaleFactorButton_OnClicked()
-{
-	TSharedPtr<SWindow> OwnerWindow = OwnerWindowPtr.Pin();
-
-	if (OwnerWindow.IsValid())
-	{
-		ScaleFactor = 1.0f;
-		ApplyWindowScaleFactor(OwnerWindow);
-	}
-
-	return FReply::Handled();
 }
 
 #endif
\ No newline at end of file
diff --git a/Engine/Source/Editor/PIEPreviewDeviceProfileSelector/Private/PIEPreviewWindowTitleBar.h b/Engine/Source/Editor/PIEPreviewDeviceProfileSelector/Private/PIEPreviewWindowTitleBar.h
--- a/Engine/Source/Editor/PIEPreviewDeviceProfileSelector/Private/PIEPreviewWindowTitleBar.h
+++ b/Engine/Source/Editor/PIEPreviewDeviceProfileSelector/Private/PIEPreviewWindowTitleBar.h
@@ -14,52 +14,30 @@
 class PIEPREVIEWDEVICEPROFILESELECTOR_API SPIEPreviewWindowTitleBar
 	: public SWindowTitleBar
 {
-
 private:
 	/**
 	* Creates widgets for this window's title bar area.
 	*/
 	virtual void MakeTitleBarContentWidgets(TSharedPtr< SWidget >& OutLeftContent, TSharedPtr< SWidget >& OutRightContent) override;
-	
-	void ApplyWindowRotation(TSharedPtr<SWindow> OwnerWindow);
-
-	void ApplyWindowScaleFactor(TSharedPtr<SWindow> OwnerWindow);
 
 	const FSlateBrush* GetScreenRotationButtonImage() const;
 
-	const FSlateBrush* GetQuarterMobileContentScaleFactorImage() const;
-	
-	const FSlateBrush* GetHalfMobileContentScaleFactorImage() const;
-
-	const FSlateBrush* GetFullMobileContentScaleFactorImage() const;
-
-	FReply ScreenRotationButton_OnClicked();
-
-
-	FReply QuarterMobileContentScaleFactorButton_OnClicked();
-
-
-	FReply HalfMobileContentScaleFactorButton_OnClicked();
-
-
-	FReply FullMobileContentScaleFactorButton_OnClicked();
+	TSharedPtr<class SPIEPreviewWindow> GetOwnerWindow() const;
 
 private:
-
 	// Holds the screen rotation button.
 	TSharedPtr<SButton> ScreenRotationButton;
 
-	// Holds the 0.25X button.
-	TSharedPtr<SButton> QuarterMobileContentScaleFactorButton;
-
-	// Holds the 0.5X button.
-	TSharedPtr<SButton> HalfMobileContentScaleFactorButton;
+	// Holds the clamp button.
+	TSharedPtr<class SCheckBox> ClampWindowSizeCheckBox;
+};
 
-	// Holds the 1X button.
-	TSharedPtr<SButton> FullMobileContentScaleFactorButton;
+FORCEINLINE TSharedPtr<class SPIEPreviewWindow> SPIEPreviewWindowTitleBar::GetOwnerWindow() const
+{
+	TSharedPtr<SWindow> OwnerWindow = OwnerWindowPtr.Pin();
+	TSharedPtr<SPIEPreviewWindow> PIEWindow = StaticCastSharedPtr<SPIEPreviewWindow>(OwnerWindow);
 
-	bool IsPortrait;
+	return PIEWindow;
+}
 
-	float ScaleFactor = 1.0f;
-};
 #endif
diff --git a/Engine/Source/Editor/PIEPreviewDeviceProfileSelector/Public/PIEPreviewDeviceProfileSelectorModule.h b/Engine/Source/Editor/PIEPreviewDeviceProfileSelector/Public/PIEPreviewDeviceProfileSelectorModule.h
--- a/Engine/Source/Editor/PIEPreviewDeviceProfileSelector/Public/PIEPreviewDeviceProfileSelectorModule.h
+++ b/Engine/Source/Editor/PIEPreviewDeviceProfileSelector/Public/PIEPreviewDeviceProfileSelectorModule.h
@@ -7,8 +7,6 @@
 #include "RHIDefinitions.h"
 #include "Misc/CommandLine.h"
 #include "Widgets/SWindow.h"
-#include "../Private/PIEPreviewWindowCoreStyle.h"
-#include "../Private/PIEPreviewWindow.h"
 #include "IPIEPreviewDeviceModule.h"
 /**
 * Implements the Preview Device Profile Selector module.
@@ -42,6 +40,9 @@ public:
 	virtual void ApplyPreviewDeviceState() override;
 	
 	virtual TSharedRef<SWindow> CreatePIEPreviewDeviceWindow(FVector2D ClientSize, FText WindowTitle, EAutoCenter AutoCenterType, FVector2D ScreenPosition, TOptional<float> MaxWindowWidth, TOptional<float> MaxWindowHeight) override;
+
+	/** call this after the window is created and registered to the application to setup display related parameters */
+	virtual void PrepareDeviceDisplay() override;
 	
 	virtual const FPIEPreviewDeviceContainer& GetPreviewDeviceContainer() ;
 	TSharedPtr<FPIEPreviewDeviceContainerCategory> GetPreviewDeviceRootCategory() const { return EnumeratedDevices.GetRootCategory(); }
@@ -58,18 +59,21 @@ private:
 		return TEXT("MobileTargetDevice=");
 	}
 
-	const void GetPreviewDeviceResolution(int32& ScreenWidth, int32& ScreenHeight);
-
 	void InitPreviewDevice();
 	static FString GetDeviceSpecificationContentDir();
 	bool ReadDeviceSpecification();
-	ERHIFeatureLevel::Type GetPreviewDeviceFeatureLevel() const;
 	FString FindDeviceSpecificationFilePath(const FString& SearchDevice);
 
+	void UpdateDisplayResolution();
+
+private:
 	bool bInitialized;
 	FString DeviceProfile;
 	FString PreviewDevice;
 
 	FPIEPreviewDeviceContainer EnumeratedDevices;
-	TSharedPtr<struct FPIEPreviewDeviceSpecifications> DeviceSpecs;
+
+	TSharedPtr<class FPIEPreviewDevice> Device;
+
+	TWeakPtr<class SPIEPreviewWindow> WindowWPtr;
 };
diff --git a/Engine/Source/Runtime/ApplicationCore/Private/Android/AndroidWindow.cpp b/Engine/Source/Runtime/ApplicationCore/Private/Android/AndroidWindow.cpp
--- a/Engine/Source/Runtime/ApplicationCore/Private/Android/AndroidWindow.cpp
+++ b/Engine/Source/Runtime/ApplicationCore/Private/Android/AndroidWindow.cpp
@@ -1,6 +1,7 @@
 // Copyright 1998-2018 Epic Games, Inc. All Rights Reserved.
 
 #include "Android/AndroidWindow.h"
+#include "Android/AndroidWindowUtils.h"
 #if USE_ANDROID_JNI
 #include <android/native_window.h>
 #include <android/native_window_jni.h>
@@ -122,6 +123,53 @@ void* FAndroidWindow::GetHardwareWindow()
 #if USE_ANDROID_JNI
 extern bool AndroidThunkCpp_IsGearVRApplication();
 #endif
+
+bool FAndroidWindow::IsCachedRectValid(const bool bMosaicEnabled, const float RequestedContentScaleFactor, ANativeWindow* Window)
+{
+	if (!WindowInit)
+	{
+		return false;
+	}
+
+	bool bValidCache = true;
+
+	if (bLastMosaicState != bMosaicEnabled)
+	{
+		FPlatformMisc::LowLevelOutputDebugStringf(TEXT("***** Mosaic State change (to %s), not using res cache"), bMosaicEnabled ? TEXT("enabled") : TEXT("disabled"));
+		bValidCache = false;
+	}
+
+	if (RequestedContentScaleFactor != ContentScaleFactor)
+	{
+		FPlatformMisc::LowLevelOutputDebugStringf(TEXT("***** RequestedContentScaleFactor different %f != %f, not using res cache"), RequestedContentScaleFactor, ContentScaleFactor);
+		bValidCache = false;
+	}
+
+	if (Window != LastWindow)
+	{
+		FPlatformMisc::LowLevelOutputDebugString(TEXT("***** Window different, not using res cache"));
+		bValidCache = false;
+	}
+
+	if (WindowWidth <= 8)
+	{
+		FPlatformMisc::LowLevelOutputDebugStringf(TEXT("***** WindowWidth is %d, not using res cache"), WindowWidth);
+		bValidCache = false;
+	}
+
+	return bValidCache;
+}
+
+void FAndroidWindow::CacheRect(ANativeWindow* Window, const int32 Width, const int32 Height, const float RequestedContentScaleFactor, const bool bMosaicEnabled)
+{
+	WindowWidth = Width;
+	WindowHeight = Height;
+	WindowInit = true;
+	ContentScaleFactor = RequestedContentScaleFactor;
+	LastWindow = Window;
+	bLastMosaicState = bMosaicEnabled;
+}
+
 FPlatformRect FAndroidWindow::GetScreenRect()
 {
 	int32 OverrideResX, OverrideResY;
@@ -143,18 +191,8 @@ FPlatformRect FAndroidWindow::GetScreenRect()
 	return FPlatformRect();
 #else
 
-	// CSF is a multiplier to 1280x720
-	static IConsoleVariable* CVar = IConsoleManager::Get().FindConsoleVariable(TEXT("r.MobileContentScaleFactor"));
 	static const bool bIsGearVRApp = AndroidThunkCpp_IsGearVRApplication();
-	// If the app is for Gear VR then always use 0 as ScaleFactor (to match window size).
-	float RequestedContentScaleFactor = (!bIsGearVRApp) ? CVar->GetFloat() : 0;
 
-	FString CmdLineCSF;
-	if (FParse::Value(FCommandLine::Get(), TEXT("mcsf="), CmdLineCSF, false))
-	{
-		RequestedContentScaleFactor = FCString::Atof(*CmdLineCSF);
-	}
-	
 	ANativeWindow* Window = (ANativeWindow*)FAndroidWindow::GetHardwareWindow();
 	static const bool bIsDaydreamApp = FAndroidMisc::IsDaydreamApplication();
 	if (bIsDaydreamApp && Window == NULL)
@@ -172,7 +210,7 @@ FPlatformRect FAndroidWindow::GetScreenRect()
 			Window = (ANativeWindow*)FAndroidWindow::GetHardwareWindow();
 		}
 	}
-//	check(Window != NULL);
+	//	check(Window != NULL);
 	if (Window == NULL)
 	{
 		FPlatformRect ScreenRect;
@@ -187,167 +225,59 @@ FPlatformRect FAndroidWindow::GetScreenRect()
 	}
 
 	// determine mosaic requirements:
-	static auto* MobileHDRCvar = IConsoleManager::Get().FindTConsoleVariableDataInt(TEXT("r.MobileHDR"));
-	const bool bMobileHDR = (MobileHDRCvar && MobileHDRCvar->GetValueOnAnyThread() == 1);
-
-	static auto* MobileHDR32bppModeCvar = IConsoleManager::Get().FindTConsoleVariableDataInt(TEXT("r.MobileHDR32bppMode"));
-	const int32 MobileHDR32Mode = MobileHDR32bppModeCvar->GetValueOnAnyThread();
-
-	bool bMosaicEnabled = false;
-	bool bHDR32ModeOverridden = false;
-	bool bDeviceRequiresHDR32bpp = false;
-	bool bDeviceRequiresMosaic = false;
-	if (!bIsGearVRApp && !bIsDaydreamApp)
-	{
-		bDeviceRequiresHDR32bpp = !FAndroidMisc::SupportsFloatingPointRenderTargets();
-		bDeviceRequiresMosaic = bDeviceRequiresHDR32bpp && !FAndroidMisc::SupportsShaderFramebufferFetch();
-
-		bHDR32ModeOverridden = MobileHDR32Mode != 0;
-		bMosaicEnabled = bDeviceRequiresMosaic && (!bHDR32ModeOverridden || MobileHDR32Mode == 1);
-	}
-
-	bool bUseResCache = WindowInit;
-
-	if (bLastMosaicState != bMosaicEnabled)
-	{
-		FPlatformMisc::LowLevelOutputDebugStringf(TEXT("***** Mosaic State change (to %s), not using res cache"), bMosaicEnabled ? TEXT("enabled") : TEXT("disabled"));
-		bUseResCache = false;
-	}
-	
-	if (RequestedContentScaleFactor != ContentScaleFactor)
-	{
-		FPlatformMisc::LowLevelOutputDebugStringf(TEXT("***** RequestedContentScaleFactor different %f != %f, not using res cache"), RequestedContentScaleFactor, ContentScaleFactor);
-		bUseResCache = false;
-	}
+	const bool bMosaicEnabled = AndroidWindowUtils::ShouldEnableMosaic() && !(bIsGearVRApp || bIsDaydreamApp);
 
-	if (Window != LastWindow)
-	{
-		FPlatformMisc::LowLevelOutputDebugString(TEXT("***** Window different, not using res cache"));
-		bUseResCache = false;
-	}
+	// CSF is a multiplier to 1280x720
+	static IConsoleVariable* CVar = IConsoleManager::Get().FindConsoleVariable(TEXT("r.MobileContentScaleFactor"));
+	// If the app is for Gear VR then always use 0 as ScaleFactor (to match window size).
+	float RequestedContentScaleFactor = bIsGearVRApp ? 0.0f : CVar->GetFloat();
 
-	if (WindowWidth <= 8)
+	FString CmdLineCSF;
+	if (FParse::Value(FCommandLine::Get(), TEXT("mcsf="), CmdLineCSF, false))
 	{
-		FPlatformMisc::LowLevelOutputDebugStringf(TEXT("***** WindowWidth is %d, not using res cache"), WindowWidth);
-		bUseResCache = false;
+		RequestedContentScaleFactor = FCString::Atof(*CmdLineCSF);
 	}
 
 	// since orientation won't change on Android, use cached results if still valid
-	if (bUseResCache)
+	bool bComputeRect = !IsCachedRectValid(bMosaicEnabled, RequestedContentScaleFactor, Window);
+	if (bComputeRect)
 	{
-		FPlatformRect ScreenRect;
-		ScreenRect.Left = 0;
-		ScreenRect.Top = 0;
-		ScreenRect.Right = WindowWidth;
-		ScreenRect.Bottom = WindowHeight;
-
-		return ScreenRect;
-	}
-
-	// currently hardcoding resolution
-
-	// get the aspect ratio of the physical screen
-	int32 ScreenWidth, ScreenHeight;
-	CalculateSurfaceSize(Window, ScreenWidth, ScreenHeight);
-	float AspectRatio = (float)ScreenWidth / (float)ScreenHeight;
-
-	int32 MaxWidth = ScreenWidth; 
-	int32 MaxHeight = ScreenHeight;
+		// currently hardcoding resolution
 
+		// get the aspect ratio of the physical screen
+		int32 ScreenWidth, ScreenHeight;
+		CalculateSurfaceSize(Window, ScreenWidth, ScreenHeight);
 
-	UE_LOG(LogAndroid, Log, TEXT("Mobile HDR: %s"), bMobileHDR ? TEXT("YES") : TEXT("no"));
-	if (bMobileHDR && !bIsGearVRApp)
-	{
-		UE_LOG(LogAndroid, Log, TEXT("Device requires 32BPP mode : %s"), bDeviceRequiresHDR32bpp ? TEXT("YES") : TEXT("no"));
-		UE_LOG(LogAndroid, Log, TEXT("Device requires mosaic: %s"), bDeviceRequiresMosaic ? TEXT("YES") : TEXT("no"));
-
-		if(bHDR32ModeOverridden)
-		{
-			UE_LOG(LogAndroid, Log, TEXT("--- Enabling 32 BPP override with 'r.MobileHDR32bppMode' = %d"), MobileHDR32Mode);
-			UE_LOG(LogAndroid, Log, TEXT("  32BPP mode : YES"));
-			UE_LOG(LogAndroid, Log, TEXT("  32BPP mode requires mosaic: %s"), bMosaicEnabled ? TEXT("YES") : TEXT("no"));
-			UE_LOG(LogAndroid, Log, TEXT("  32BPP mode requires RGBE: %s"), MobileHDR32Mode == 2 ? TEXT("YES") : TEXT("no"));
-		}
+		static auto* MobileHDRCvar = IConsoleManager::Get().FindTConsoleVariableDataInt(TEXT("r.MobileHDR"));
+		const bool bMobileHDR = (MobileHDRCvar && MobileHDRCvar->GetValueOnAnyThread() == 1);
+		UE_LOG(LogAndroid, Log, TEXT("Mobile HDR: %s"), bMobileHDR ? TEXT("YES") : TEXT("no"));
 
-		if(bMosaicEnabled)
+		if (!bIsGearVRApp)
 		{
-			UE_LOG(LogAndroid, Log, TEXT("Using mosaic rendering due to lack of Framebuffer Fetch support."));
-			if (!FAndroidMisc::SupportsES30())
+			bool bSupportsES30 = FAndroidMisc::SupportsES30();
+			if (!bIsDaydreamApp && !bSupportsES30)
 			{
-				const int32 OldMaxWidth = MaxWidth;
-				const int32 OldMaxHeight = MaxHeight;
-
-				if (GAndroidIsPortrait)
-				{
-					MaxHeight = FPlatformMath::Min(MaxHeight, 1024);
-					MaxWidth = (MaxHeight * AspectRatio + 0.5f);
-				}
-				else
-				{
-					MaxWidth = FPlatformMath::Min(MaxWidth, 1024);
-					MaxHeight = (MaxWidth / AspectRatio + 0.5f);
-				}
-
-				// ensure Width and Height is multiple of 8
-				MaxWidth = (MaxWidth / 8) * 8; 
-				MaxHeight = (MaxHeight / 8) * 8;
-
-				UE_LOG(LogAndroid, Log, TEXT("Limiting MaxWidth=%d and MaxHeight=%d due to mosaic rendering on ES2 device (was %dx%d)"), MaxWidth, MaxHeight, OldMaxWidth, OldMaxHeight);
+				AndroidWindowUtils::ApplyMosaicRequirements(ScreenWidth, ScreenHeight);
 			}
-		}
-	}
 
-	// 0 means to use native size
-	int32 Width, Height;
-	if (RequestedContentScaleFactor == 0.0f)
-	{
-		Width = MaxWidth;
-		Height = MaxHeight;
-		UE_LOG(LogAndroid, Log, TEXT("Setting Width=%d and Height=%d (requested scale = 0 = auto)"), Width, Height);
-	}
-	else
-	{
-		if (GAndroidIsPortrait)
-		{
-			Height = 1280 * RequestedContentScaleFactor;
-		}
-		else
-		{
-			Height = 720 * RequestedContentScaleFactor;
+			AndroidWindowUtils::ApplyContentScaleFactor(ScreenWidth, ScreenHeight);
 		}
 
-		// apply the aspect ration to get the width
-		Width = (Height * AspectRatio + 0.5f);
-		// ensure Width and Height is multiple of 8
-		Width = (Width / 8) * 8; 
-		Height = (Height / 8) * 8;
-
-		// clamp to native resolution
-		Width = FPlatformMath::Min(Width, MaxWidth);
-		Height = FPlatformMath::Min(Height, MaxHeight);
-
-		UE_LOG(LogAndroid, Log, TEXT("Setting Width=%d and Height=%d (requested scale = %f)"), Width, Height, RequestedContentScaleFactor);
+		// save for future calls
+		CacheRect(Window, ScreenWidth, ScreenHeight, RequestedContentScaleFactor, bMosaicEnabled);
 	}
 
+	// create rect and return
 	FPlatformRect ScreenRect;
 	ScreenRect.Left = 0;
 	ScreenRect.Top = 0;
-	ScreenRect.Right = Width;
-	ScreenRect.Bottom = Height;
-
-	// save for future calls
-	WindowWidth = Width;
-	WindowHeight = Height;
-	WindowInit = true;
-	ContentScaleFactor = RequestedContentScaleFactor;
-	LastWindow = Window;
-	bLastMosaicState = bMosaicEnabled;
+	ScreenRect.Right = WindowWidth;
+	ScreenRect.Bottom = WindowHeight;
 
 	return ScreenRect;
 #endif
 }
 
-
 void FAndroidWindow::CalculateSurfaceSize(void* InWindow, int32_t& SurfaceWidth, int32_t& SurfaceHeight)
 {
 	// allow a subplatform to dictate resolution - we can't easily subclass FAndroidWindow the way its used
@@ -358,7 +288,7 @@ void FAndroidWindow::CalculateSurfaceSize(void* InWindow, int32_t& SurfaceWidth,
 
 	// too much of the following code needs JNI things, just assume override
 #if !USE_ANDROID_JNI
-	
+
 	UE_LOG(LogAndroid, Fatal, TEXT("FAndroidWindow::CalculateSurfaceSize currently expedcts non-JNI platforms to override resolution"));
 
 #else
@@ -416,7 +346,7 @@ void FAndroidWindow::CalculateSurfaceSize(void* InWindow, int32_t& SurfaceWidth,
 	}
 
 	// some phones gave it the other way (so, if swap if the app is landscape, but width < height)
-	if ((GAndroidIsPortrait && SurfaceWidth > SurfaceHeight) || 
+	if ((GAndroidIsPortrait && SurfaceWidth > SurfaceHeight) ||
 		(!GAndroidIsPortrait && SurfaceWidth < SurfaceHeight))
 	{
 		Swap(SurfaceWidth, SurfaceHeight);
@@ -424,7 +354,7 @@ void FAndroidWindow::CalculateSurfaceSize(void* InWindow, int32_t& SurfaceWidth,
 
 	// ensure the size is divisible by a specified amount
 	// do not convert to a surface size that is larger than native resolution
-	// Mobile VR doesn�t need buffer quantization as UE4 never renders directly to the buffer in VR mode. 
+	// Mobile VR doesn't need buffer quantization as UE4 never renders directly to the buffer in VR mode. 
 	const int DividableBy = bIsMobileVRApp ? 1 : 8;
 	SurfaceWidth = (SurfaceWidth / DividableBy) * DividableBy;
 	SurfaceHeight = (SurfaceHeight / DividableBy) * DividableBy;
diff --git a/Engine/Source/Runtime/ApplicationCore/Private/GenericPlatform/GenericWindow.cpp b/Engine/Source/Runtime/ApplicationCore/Private/GenericPlatform/GenericWindow.cpp
--- a/Engine/Source/Runtime/ApplicationCore/Private/GenericPlatform/GenericWindow.cpp
+++ b/Engine/Source/Runtime/ApplicationCore/Private/GenericPlatform/GenericWindow.cpp
@@ -174,3 +174,19 @@ float FGenericWindow::GetDPIScaleFactor() const
 {
 	return 1.0f;
 }
+
+void FGenericWindow::SetDPIScaleFactor(const float Factor)
+{
+	// empty default functionality
+}
+
+bool FGenericWindow::IsManualManageDPIChanges() const
+{
+	// returns false by default so the application can auto-manage the size of its windows, in response to DPI variations
+	return false;
+}
+
+void FGenericWindow::SetManualManageDPIChanges(const bool bAutoHandle)
+{
+	// empty default functionality
+}
\ No newline at end of file
diff --git a/Engine/Source/Runtime/ApplicationCore/Private/Linux/LinuxApplication.cpp b/Engine/Source/Runtime/ApplicationCore/Private/Linux/LinuxApplication.cpp
--- a/Engine/Source/Runtime/ApplicationCore/Private/Linux/LinuxApplication.cpp
+++ b/Engine/Source/Runtime/ApplicationCore/Private/Linux/LinuxApplication.cpp
@@ -1461,6 +1461,16 @@ void FDisplayMetrics::GetDisplayMetrics(FDisplayMetrics& OutDisplayMetrics)
 		Display.DisplayRect = FPlatformRect(DisplayBounds.x, DisplayBounds.y, DisplayBounds.x + DisplayBounds.w, DisplayBounds.y + DisplayBounds.h);
 		Display.WorkArea = FPlatformRect(UsableBounds.x, UsableBounds.y, UsableBounds.x + UsableBounds.w, UsableBounds.y + UsableBounds.h);
 		Display.bIsPrimary = DisplayIdx == 0;
+
+		if (FPlatformApplicationMisc::IsHighDPIAwarenessEnabled())
+		{
+			float HorzDPI = 0.0f, VertDPI = 0.0f;
+			if (SDL_GetDisplayDPI(DisplayIdx, nullptr, &HorzDPI, &VertDPI) == 0)
+			{
+				Display.DPI = FMath::FloorToInt((HorzDPI + VertDPI) / 2.0f);
+			}
+		}
+
 		OutDisplayMetrics.MonitorInfo.Add(Display);
 
 		if (Display.bIsPrimary)
diff --git a/Engine/Source/Runtime/ApplicationCore/Private/Mac/MacApplication.cpp b/Engine/Source/Runtime/ApplicationCore/Private/Mac/MacApplication.cpp
--- a/Engine/Source/Runtime/ApplicationCore/Private/Mac/MacApplication.cpp
+++ b/Engine/Source/Runtime/ApplicationCore/Private/Mac/MacApplication.cpp
@@ -1965,12 +1965,42 @@ void FDisplayMetrics::GetDisplayMetrics(FDisplayMetrics& OutDisplayMetrics)
 
 		FMonitorInfo Info;
 		Info.ID = FString::Printf(TEXT("%u"), DisplayID);
-		Info.NativeWidth = CGDisplayPixelsWide(DisplayID);
-		Info.NativeHeight = CGDisplayPixelsHigh(DisplayID);
+
+		CFArrayRef ArrDisplay = CGDisplayCopyAllDisplayModes(DisplayID, nullptr);
+
+		Info.NativeWidth = 0;
+		Info.NativeHeight = 0;
+		const CFIndex AppsCount = CFArrayGetCount(ArrDisplay);
+		for (CFIndex i = 0; i < AppsCount; ++i)
+		{
+			const CGDisplayModeRef Mode = (const CGDisplayModeRef)CFArrayGetValueAtIndex(ArrDisplay, i);
+			const int32 Width = (int32)CGDisplayModeGetWidth(Mode);
+			const int32 Height = (int32)CGDisplayModeGetHeight(Mode);
+
+			if (Width * Height > Info.NativeWidth * Info.NativeHeight)
+			{
+				Info.NativeWidth = Width;
+				Info.NativeHeight = Height;
+			}
+		}
+
+		if (!Info.NativeWidth || !Info.NativeHeight)
+		{
+			Info.NativeWidth = CGDisplayPixelsWide(DisplayID);
+			Info.NativeHeight = CGDisplayPixelsHigh(DisplayID);
+		}
+
 		Info.DisplayRect = FPlatformRect(Screen->FramePixels.origin.x, Screen->FramePixels.origin.y, Screen->FramePixels.origin.x + Screen->FramePixels.size.width, Screen->FramePixels.origin.y + Screen->FramePixels.size.height);
 		Info.WorkArea = FPlatformRect(Screen->VisibleFramePixels.origin.x, Screen->VisibleFramePixels.origin.y, Screen->VisibleFramePixels.origin.x + Screen->VisibleFramePixels.size.width, Screen->VisibleFramePixels.origin.y + Screen->VisibleFramePixels.size.height);
 		Info.bIsPrimary = Screen->Screen == [NSScreen mainScreen];
 
+		// dpi computations
+		const CGSize DisplayPhysicalSize = CGDisplayScreenSize(DisplayID);
+		const float MilimetreInch = 25.4f;
+		float HorizontalDPI = MilimetreInch * (float)Info.NativeWidth / (float)DisplayPhysicalSize.width;
+		float VerticalDPI = MilimetreInch * (float)Info.NativeHeight / (float)DisplayPhysicalSize.height;
+		Info.DPI = FMath::CeilToInt((HorizontalDPI + VerticalDPI) / 2.0f);
+
 		// Monitor's name can only be obtained from IOKit
 		io_iterator_t IOIterator;
 		kern_return_t Result = IOServiceGetMatchingServices(kIOMasterPortDefault, IOServiceMatching("IODisplayConnect"), &IOIterator);
diff --git a/Engine/Source/Runtime/ApplicationCore/Private/Windows/WindowsApplication.cpp b/Engine/Source/Runtime/ApplicationCore/Private/Windows/WindowsApplication.cpp
--- a/Engine/Source/Runtime/ApplicationCore/Private/Windows/WindowsApplication.cpp
+++ b/Engine/Source/Runtime/ApplicationCore/Private/Windows/WindowsApplication.cpp
@@ -633,6 +633,9 @@ static void GetMonitorsInfo(TArray<FMonitorInfo>& OutMonitorInfo)
 					{
 						Info.ID = Monitor.DeviceID;
 						Info.bIsPrimary = (DisplayDevice.StateFlags & DISPLAY_DEVICE_PRIMARY_DEVICE) > 0;
+
+						Info.DPI = FWindowsPlatformApplicationMisc::GetMonitorDPI(Info);
+
 						OutMonitorInfo.Add(Info);
 
 						if (PrimaryDevice == nullptr && Info.bIsPrimary)
@@ -2032,15 +2035,19 @@ int32 FWindowsApplication::ProcessDeferredMessage( const FDeferredWindowsMessage
 
 		case WM_DPICHANGED:
 			{
-				if( CurrentNativeEventWindowPtr.IsValid() )
+				if( CurrentNativeEventWindowPtr.IsValid())
 				{
-					CurrentNativeEventWindowPtr->SetDPIScaleFactor(LOWORD(wParam) / 96.0f);
+					MessageHandler->SignalSystemDPIChanged(CurrentNativeEventWindowPtr.ToSharedRef());
 
+					if (!CurrentNativeEventWindowPtr->IsManualManageDPIChanges())
+					{
+						CurrentNativeEventWindowPtr->SetDPIScaleFactor(LOWORD(wParam) / 96.0f);
 
-					LPRECT NewRect = (LPRECT)lParam;
-					SetWindowPos(hwnd, nullptr, NewRect->left, NewRect->top, NewRect->right - NewRect->left, NewRect->bottom - NewRect->top, SWP_NOZORDER | SWP_NOACTIVATE);
+						LPRECT NewRect = (LPRECT)lParam;
+						SetWindowPos(hwnd, nullptr, NewRect->left, NewRect->top, NewRect->right - NewRect->left, NewRect->bottom - NewRect->top, SWP_NOZORDER | SWP_NOACTIVATE);
 
-					MessageHandler->HandleDPIScaleChanged(CurrentNativeEventWindowPtr.ToSharedRef());
+						MessageHandler->HandleDPIScaleChanged(CurrentNativeEventWindowPtr.ToSharedRef());
+					}
 				}
 			}
 			break;
diff --git a/Engine/Source/Runtime/ApplicationCore/Private/Windows/WindowsPlatformApplicationMisc.cpp b/Engine/Source/Runtime/ApplicationCore/Private/Windows/WindowsPlatformApplicationMisc.cpp
--- a/Engine/Source/Runtime/ApplicationCore/Private/Windows/WindowsPlatformApplicationMisc.cpp
+++ b/Engine/Source/Runtime/ApplicationCore/Private/Windows/WindowsPlatformApplicationMisc.cpp
@@ -274,8 +274,46 @@ bool FWindowsPlatformApplicationMisc::GetWindowTitleMatchingText(const TCHAR* Ti
 	return bWasFound;
 }
 
+int32 FWindowsPlatformApplicationMisc::GetMonitorDPI(const FMonitorInfo& MonitorInfo)
+{
+	int32 DisplayDPI = 96;
+
+	if (IsHighDPIAwarenessEnabled())
+	{
+		if (GetDpiForMonitor)
+		{
+			RECT MonitorDim;
+			MonitorDim.left = MonitorInfo.DisplayRect.Left;
+			MonitorDim.top = MonitorInfo.DisplayRect.Top;
+			MonitorDim.right = MonitorInfo.DisplayRect.Right;
+			MonitorDim.bottom = MonitorInfo.DisplayRect.Bottom;
+
+			HMONITOR Monitor = MonitorFromRect(&MonitorDim, MONITOR_DEFAULTTONEAREST);
+			if (Monitor)
+			{
+				uint32 DPIX = 0;
+				uint32 DPIY = 0;
+				if (SUCCEEDED(GetDpiForMonitor(Monitor, 0 /*MDT_EFFECTIVE_DPI*/, &DPIX, &DPIY)))
+				{
+					DisplayDPI = DPIX;
+				}
+			}
+		}
+		else
+		{
+			HDC Context = GetDC(nullptr);
+			DisplayDPI = GetDeviceCaps(Context, LOGPIXELSX);
+			ReleaseDC(nullptr, Context);
+		}
+	}
+
+	return DisplayDPI;
+}
+
 float FWindowsPlatformApplicationMisc::GetDPIScaleFactorAtPoint(float X, float Y)
 {
+	float Scale = 1.0f;
+
 	if (IsHighDPIAwarenessEnabled())
 	{
 		if (GetDpiForMonitor)
@@ -286,18 +324,22 @@ float FWindowsPlatformApplicationMisc::GetDPIScaleFactorAtPoint(float X, float Y
 			{
 				uint32 DPIX = 0;
 				uint32 DPIY = 0;
-				return SUCCEEDED(GetDpiForMonitor(Monitor, 0 /*MDT_EFFECTIVE_DPI*/, &DPIX, &DPIY)) ? DPIX / 96.0f : 1.0f;
+				if (SUCCEEDED(GetDpiForMonitor(Monitor, 0 /*MDT_EFFECTIVE_DPI*/, &DPIX, &DPIY)))
+				{
+					Scale = (float)DPIX / 96.0f;
+				}
 			}
 		}
 		else
 		{
 			HDC Context = GetDC(nullptr);
-			const float DPIScaleFactor = GetDeviceCaps(Context, LOGPIXELSX) / 96.0f;
+			int32 DPI = GetDeviceCaps(Context, LOGPIXELSX);
+			Scale = (float)DPI / 96.0f;
 			ReleaseDC(nullptr, Context);
-			return DPIScaleFactor;
 		}
 	}
-	return 1.0f;
+
+	return Scale;
 }
 
 // Disabling optimizations helps to reduce the frequency of OpenClipboard failing with error code 0. It still happens
diff --git a/Engine/Source/Runtime/ApplicationCore/Private/Windows/WindowsWindow.cpp b/Engine/Source/Runtime/ApplicationCore/Private/Windows/WindowsWindow.cpp
--- a/Engine/Source/Runtime/ApplicationCore/Private/Windows/WindowsWindow.cpp
+++ b/Engine/Source/Runtime/ApplicationCore/Private/Windows/WindowsWindow.cpp
@@ -844,6 +844,15 @@ bool FWindowsWindow::GetRestoredDimensions(int32& X, int32& Y, int32& Width, int
 	}
 }
 
+bool FWindowsWindow::IsManualManageDPIChanges() const
+{
+	return bHandleManualDPIChanges;
+}
+
+void FWindowsWindow::SetManualManageDPIChanges(const bool bManualDPIChanges)
+{
+	bHandleManualDPIChanges = bManualDPIChanges;
+}
 
 void FWindowsWindow::AdjustCachedSize( FVector2D& Size ) const
 {
diff --git a/Engine/Source/Runtime/ApplicationCore/Public/Android/AndroidWindow.h b/Engine/Source/Runtime/ApplicationCore/Public/Android/AndroidWindow.h
--- a/Engine/Source/Runtime/ApplicationCore/Public/Android/AndroidWindow.h
+++ b/Engine/Source/Runtime/ApplicationCore/Public/Android/AndroidWindow.h
@@ -35,7 +35,6 @@ public:
 	/** Returns the rectangle of the screen the window is associated with */
 	virtual bool GetFullScreenInfo( int32& X, int32& Y, int32& Width, int32& Height ) const override;
 
-	
 	virtual void SetOSWindowHandle(void*);
 
 	static FPlatformRect GetScreenRect();
@@ -57,6 +56,12 @@ protected:
 	virtual EWindowMode::Type GetWindowMode() const override { return EWindowMode::Fullscreen; }
 
 private:
+	/** called from GetScreenRect function */
+	/** test cached values from the latest computations stored by CacheRect to decide their validity with the provided arguments */
+	static bool IsCachedRectValid(const bool bMosaicEnabled, const float RequestedContentScaleFactor, ANativeWindow* Window);
+	/** caches some values used to compute the size of the window by GetScreenRect function */
+	static void CacheRect(ANativeWindow* Window, const int32 Width, const int32 Height, const float RequestedContentScaleFactor, const bool bMosaicEnabled);
+
 	/**
 	 * Protect the constructor; only TSharedRefs of this class can be made.
 	 */
diff --git a/Engine/Source/Runtime/ApplicationCore/Public/Android/AndroidWindowUtils.h b/Engine/Source/Runtime/ApplicationCore/Public/Android/AndroidWindowUtils.h
new file mode 100644
--- /dev/null
+++ b/Engine/Source/Runtime/ApplicationCore/Public/Android/AndroidWindowUtils.h
@@ -0,0 +1,142 @@
+#pragma once
+
+#include "HAL/IConsoleManager.h"
+#include "RHI.h"
+
+#if PLATFORM_ANDROID
+#include "Android/AndroidMisc.h"
+#endif
+
+namespace AndroidWindowUtils
+{
+	inline bool DeviceRequiresMosaic()
+	{
+#if PLATFORM_ANDROID
+		bool bDeviceRequiresMosaic = !FAndroidMisc::SupportsFloatingPointRenderTargets() && !FAndroidMisc::SupportsShaderFramebufferFetch();
+#else
+		bool bDeviceRequiresMosaic = !GSupportsRenderTargetFormat_PF_FloatRGBA && !GSupportsShaderFramebufferFetch;
+#endif
+
+		return bDeviceRequiresMosaic;
+	}
+
+	inline bool ShouldEnableMosaic()
+	{
+		static auto* MobileHDR32bppModeCvar = IConsoleManager::Get().FindTConsoleVariableDataInt(TEXT("r.MobileHDR32bppMode"));
+		const int32 MobileHDR32Mode = MobileHDR32bppModeCvar->GetValueOnAnyThread();
+
+		bool bEnableMosaic = DeviceRequiresMosaic();
+		bEnableMosaic &= (MobileHDR32Mode == 0 || MobileHDR32Mode == 1);
+
+		return bEnableMosaic;
+	}
+
+	static void ApplyMosaicRequirements(int32& InOutScreenWidth, int32& InOutScreenHeight)
+	{
+		static auto* MobileHDR32bppModeCvar = IConsoleManager::Get().FindTConsoleVariableDataInt(TEXT("r.MobileHDR32bppMode"));
+		const int32 MobileHDR32Mode = MobileHDR32bppModeCvar->GetValueOnAnyThread();
+
+		bool bDeviceRequiresMosaic = DeviceRequiresMosaic();
+		bool bMosaicEnabled = ShouldEnableMosaic();
+
+#if PLATFORM_ANDROID
+		bool bDeviceRequiresHDR32bpp = !FAndroidMisc::SupportsFloatingPointRenderTargets();
+#else
+		bool bDeviceRequiresHDR32bpp = !GSupportsRenderTargetFormat_PF_FloatRGBA;
+#endif
+
+		static auto* MobileHDRCvar = IConsoleManager::Get().FindTConsoleVariableDataInt(TEXT("r.MobileHDR"));
+		const bool bMobileHDR = (MobileHDRCvar && MobileHDRCvar->GetValueOnAnyThread() == 1);
+		UE_LOG(LogAndroid, Log, TEXT("Mobile HDR: %s"), bMobileHDR ? TEXT("YES") : TEXT("no"));
+
+		if (bMobileHDR)
+		{
+			UE_LOG(LogAndroid, Log, TEXT("Device requires 32BPP mode : %s"), bDeviceRequiresHDR32bpp ? TEXT("YES") : TEXT("no"));
+			UE_LOG(LogAndroid, Log, TEXT("Device requires mosaic: %s"), bDeviceRequiresMosaic ? TEXT("YES") : TEXT("no"));
+
+			if (MobileHDR32Mode != 0)
+			{
+				UE_LOG(LogAndroid, Log, TEXT("--- Enabling 32 BPP override with 'r.MobileHDR32bppMode' = %d"), MobileHDR32Mode);
+				UE_LOG(LogAndroid, Log, TEXT("  32BPP mode : YES"));
+				UE_LOG(LogAndroid, Log, TEXT("  32BPP mode requires mosaic: %s"), bMosaicEnabled ? TEXT("YES") : TEXT("no"));
+				UE_LOG(LogAndroid, Log, TEXT("  32BPP mode requires RGBE: %s"), MobileHDR32Mode == 2 ? TEXT("YES") : TEXT("no"));
+			}
+
+			if (bMosaicEnabled)
+			{
+				UE_LOG(LogAndroid, Log, TEXT("Using mosaic rendering due to lack of Framebuffer Fetch support."));
+
+				const int32 OldScreenWidth = InOutScreenWidth;
+				const int32 OldSceenHeight = InOutScreenHeight;
+
+				const float AspectRatio = (float)InOutScreenWidth / (float)InOutScreenHeight;
+
+				if (InOutScreenHeight > InOutScreenWidth)
+				{
+					InOutScreenHeight = FPlatformMath::Min(InOutScreenHeight, 1024);
+					InOutScreenWidth = (InOutScreenHeight * AspectRatio + 0.5f);
+				}
+				else
+				{
+					InOutScreenWidth = FPlatformMath::Min(InOutScreenWidth, 1024);
+					InOutScreenHeight = (InOutScreenWidth / AspectRatio + 0.5f);
+				}
+
+				// ensure Width and Height is multiple of 8
+				InOutScreenWidth = (InOutScreenWidth / 8) * 8;
+				InOutScreenHeight = (InOutScreenHeight / 8) * 8;
+
+				UE_LOG(LogAndroid, Log, TEXT("Limiting MaxWidth=%d and MaxHeight=%d due to mosaic rendering on ES2 device (was %dx%d)"), InOutScreenWidth, InOutScreenHeight, OldScreenWidth, OldSceenHeight);
+			}
+		}
+	}
+
+	static void ApplyContentScaleFactor(int32& InOutScreenWidth, int32& InOutScreenHeight)
+	{
+		const float AspectRatio = (float)InOutScreenWidth / (float)InOutScreenHeight;
+
+		// CSF is a multiplier to 1280x720
+		static IConsoleVariable* CVar = IConsoleManager::Get().FindConsoleVariable(TEXT("r.MobileContentScaleFactor"));		
+
+		float RequestedContentScaleFactor = CVar->GetFloat();
+
+		FString CmdLineCSF;
+		if (FParse::Value(FCommandLine::Get(), TEXT("mcsf="), CmdLineCSF, false))
+		{
+			RequestedContentScaleFactor = FCString::Atof(*CmdLineCSF);
+		}
+
+		// 0 means to use native size
+		if (RequestedContentScaleFactor == 0.0f)
+		{
+			UE_LOG(LogAndroid, Log, TEXT("Setting Width=%d and Height=%d (requested scale = 0 = auto)"), InOutScreenWidth, InOutScreenHeight);
+		}
+		else
+		{
+			int32 Width = InOutScreenWidth;
+			int32 Height = InOutScreenHeight;
+
+			if (InOutScreenHeight > InOutScreenWidth)
+			{
+				Height = 1280 * RequestedContentScaleFactor;
+			}
+			else
+			{
+				Height = 720 * RequestedContentScaleFactor;
+			}
+
+			// apply the aspect ration to get the width
+			Width = (Height * AspectRatio + 0.5f);
+			// ensure Width and Height is multiple of 8
+			Width = (Width / 8) * 8;
+			Height = (Height / 8) * 8;
+
+			// clamp to native resolution
+			InOutScreenWidth = FPlatformMath::Min(Width, InOutScreenWidth);
+			InOutScreenHeight = FPlatformMath::Min(Height, InOutScreenHeight);
+
+			UE_LOG(LogAndroid, Log, TEXT("Setting Width=%d and Height=%d (requested scale = %f)"), InOutScreenWidth, InOutScreenHeight, RequestedContentScaleFactor);
+		}
+	}
+
+}
\ No newline at end of file
diff --git a/Engine/Source/Runtime/ApplicationCore/Public/GenericPlatform/GenericApplication.h b/Engine/Source/Runtime/ApplicationCore/Public/GenericPlatform/GenericApplication.h
--- a/Engine/Source/Runtime/ApplicationCore/Public/GenericPlatform/GenericApplication.h
+++ b/Engine/Source/Runtime/ApplicationCore/Public/GenericPlatform/GenericApplication.h
@@ -334,6 +334,7 @@ struct FMonitorInfo
 	FPlatformRect DisplayRect;
 	FPlatformRect WorkArea;
 	bool bIsPrimary;
+	int32 DPI = 0;
 };
 
 
diff --git a/Engine/Source/Runtime/ApplicationCore/Public/GenericPlatform/GenericApplicationMessageHandler.h b/Engine/Source/Runtime/ApplicationCore/Public/GenericPlatform/GenericApplicationMessageHandler.h
--- a/Engine/Source/Runtime/ApplicationCore/Public/GenericPlatform/GenericApplicationMessageHandler.h
+++ b/Engine/Source/Runtime/ApplicationCore/Public/GenericPlatform/GenericApplicationMessageHandler.h
@@ -409,6 +409,11 @@ public:
 
 	}
 
+	virtual void SignalSystemDPIChanged(const TSharedRef< FGenericWindow >& Window)
+	{
+
+	}	
+
 	virtual void OnMovedWindow( const TSharedRef< FGenericWindow >& Window, const int32 X, const int32 Y )
 	{
 
diff --git a/Engine/Source/Runtime/ApplicationCore/Public/GenericPlatform/GenericWindow.h b/Engine/Source/Runtime/ApplicationCore/Public/GenericPlatform/GenericWindow.h
--- a/Engine/Source/Runtime/ApplicationCore/Public/GenericPlatform/GenericWindow.h
+++ b/Engine/Source/Runtime/ApplicationCore/Public/GenericPlatform/GenericWindow.h
@@ -166,6 +166,15 @@ public:
 	 * E.g. DPIScale of 2.0 means there is a 2x2 pixel square for every 1x1 SlateUnit.
 	 */
 	virtual float GetDPIScaleFactor() const;
+
+	/** sets a new DPI scale factor */
+	virtual void SetDPIScaleFactor(const float Factor);
+
+	/** determines whether or not this window does its own DPI management */
+	virtual bool IsManualManageDPIChanges() const;
+
+	/** call with a true argument if this window need to do its custom size management in response to DPI variations */
+	virtual void SetManualManageDPIChanges(const bool bAutoHandle);
 	
 protected:
 
diff --git a/Engine/Source/Runtime/ApplicationCore/Public/GenericPlatform/GenericWindowDefinition.h b/Engine/Source/Runtime/ApplicationCore/Public/GenericPlatform/GenericWindowDefinition.h
--- a/Engine/Source/Runtime/ApplicationCore/Public/GenericPlatform/GenericWindowDefinition.h
+++ b/Engine/Source/Runtime/ApplicationCore/Public/GenericPlatform/GenericWindowDefinition.h
@@ -120,4 +120,7 @@ struct APPLICATIONCORE_API FGenericWindowDefinition
 	int32 CornerRadius;
 
 	FWindowSizeLimits SizeLimits;
+
+	/** false if the window should respond to system DPI changes, otherwise this will be handled, internally, by the application */
+	bool bManualDPI = false;
 };
diff --git a/Engine/Source/Runtime/ApplicationCore/Public/Linux/LinuxWindow.h b/Engine/Source/Runtime/ApplicationCore/Public/Linux/LinuxWindow.h
--- a/Engine/Source/Runtime/ApplicationCore/Public/Linux/LinuxWindow.h
+++ b/Engine/Source/Runtime/ApplicationCore/Public/Linux/LinuxWindow.h
@@ -169,7 +169,7 @@ public:
 		return DPIScaleFactor;
 	}
 
-	void SetDPIScaleFactor(float Value)
+	virtual void SetDPIScaleFactor(float Value) override
 	{
 		DPIScaleFactor = Value;
 	}
diff --git a/Engine/Source/Runtime/ApplicationCore/Public/Windows/WindowsPlatformApplicationMisc.h b/Engine/Source/Runtime/ApplicationCore/Public/Windows/WindowsPlatformApplicationMisc.h
--- a/Engine/Source/Runtime/ApplicationCore/Public/Windows/WindowsPlatformApplicationMisc.h
+++ b/Engine/Source/Runtime/ApplicationCore/Public/Windows/WindowsPlatformApplicationMisc.h
@@ -23,6 +23,11 @@ struct APPLICATIONCORE_API FWindowsPlatformApplicationMisc : public FGenericPlat
 	static float GetDPIScaleFactorAtPoint(float X, float Y);
 	static void ClipboardCopy(const TCHAR* Str);
 	static void ClipboardPaste(class FString& Dest);
+
+	/** Windows platform only */
+	/** Function should retrieve the DPI value for the provided monitor information structure */
+	static int32 GetMonitorDPI(const FMonitorInfo& MonitorInfo);
+	/** End Windows platform only */
 };
 
 typedef FWindowsPlatformApplicationMisc FPlatformApplicationMisc;
diff --git a/Engine/Source/Runtime/ApplicationCore/Public/Windows/WindowsWindow.h b/Engine/Source/Runtime/ApplicationCore/Public/Windows/WindowsWindow.h
--- a/Engine/Source/Runtime/ApplicationCore/Public/Windows/WindowsWindow.h
+++ b/Engine/Source/Runtime/ApplicationCore/Public/Windows/WindowsWindow.h
@@ -70,11 +70,17 @@ public:
 		return DPIScaleFactor;
 	}
 
-	void SetDPIScaleFactor(float Value)
+	virtual void SetDPIScaleFactor(float Value) override
 	{
 		DPIScaleFactor = Value;
 	}
 
+	/** determines whether or not this window does its own DPI management */
+	virtual bool IsManualManageDPIChanges() const override;
+
+	/** call with a true argument if this window need to do its custom size management in response to DPI variations */
+	virtual void SetManualManageDPIChanges(const bool bManualDPIChanges) override;
+
 	/** Called when our parent window is minimized (which will in turn cause us to become minimized). */
 	void OnParentWindowMinimized();
 
@@ -197,4 +203,7 @@ private:
 	 * E.g. DPIScale of 2.0 means there is a 2x2 pixel square for every 1x1 SlateUnit.
 	 */
 	float DPIScaleFactor;
+
+	/** when true the window is responsible for its own size adjustments in response to a DPI change */
+	bool bHandleManualDPIChanges = false;
 };
diff --git a/Engine/Source/Runtime/Engine/Private/GameEngine.cpp b/Engine/Source/Runtime/Engine/Private/GameEngine.cpp
--- a/Engine/Source/Runtime/Engine/Private/GameEngine.cpp
+++ b/Engine/Source/Runtime/Engine/Private/GameEngine.cpp
@@ -184,6 +184,17 @@ UGameEngine::UGameEngine(const FObjectInitializer& ObjectInitializer)
 void UGameEngine::CreateGameViewportWidget( UGameViewportClient* GameViewportClient )
 {
 	bool bRenderDirectlyToWindow = (!StartupMovieCaptureHandle.IsValid() || IMovieSceneCaptureModule::Get().IsStereoAllowed()) && GIsDumpingMovie == 0;
+
+	// when we're running in a "device simulation" window, render the scene to an intermediate texture
+	// in the mobile device "emulation" case this is needed to properly position the viewport (as a widget) inside its bezel
+#if WITH_EDITOR
+	auto PIEPreviewDeviceModule = FModuleManager::LoadModulePtr<IPIEPreviewDeviceModule>("PIEPreviewDeviceProfileSelector");
+	if (PIEPreviewDeviceModule && FPIEPreviewDeviceModule::IsRequestingPreviewDevice())
+	{
+		bRenderDirectlyToWindow = false;
+	}
+#endif
+
 	const bool bStereoAllowed = bRenderDirectlyToWindow;
 	TSharedRef<SOverlay> ViewportOverlayWidgetRef = SNew( SOverlay );
 
@@ -528,7 +539,7 @@ TSharedRef<SWindow> UGameEngine::CreateGameWindow()
 		Window = PIEPreviewDeviceModule->CreatePIEPreviewDeviceWindow(FVector2D(ResX, ResY), WindowTitle, AutoCenterType, FVector2D(WinX, WinY), MaxWindowWidth, MaxWindowHeight);
 	}
 #endif
-			
+
 	const bool bShowImmediately = false;
 
 	FSlateApplication::Get().AddWindow( Window, bShowImmediately );
@@ -547,6 +558,13 @@ TSharedRef<SWindow> UGameEngine::CreateGameWindow()
 
 	Window->ShowWindow();
 
+#if WITH_EDITOR
+	if (PIEPreviewDeviceModule && FPIEPreviewDeviceModule::IsRequestingPreviewDevice())
+	{
+		PIEPreviewDeviceModule->PrepareDeviceDisplay();
+	}
+#endif
+
 	// Tick now to force a redraw of the window and ensure correct fullscreen application
 	FSlateApplication::Get().Tick();
 
diff --git a/Engine/Source/Runtime/Engine/Public/IPIEPreviewDeviceModule.h b/Engine/Source/Runtime/Engine/Public/IPIEPreviewDeviceModule.h
--- a/Engine/Source/Runtime/Engine/Public/IPIEPreviewDeviceModule.h
+++ b/Engine/Source/Runtime/Engine/Public/IPIEPreviewDeviceModule.h
@@ -31,6 +31,11 @@ class IPIEPreviewDeviceModule : public IDeviceProfileSelectorModule
 		*/
 		virtual  TSharedRef<SWindow>   CreatePIEPreviewDeviceWindow(FVector2D ClientSize, FText WindowTitle, EAutoCenter AutoCenterType, FVector2D ScreenPosition, TOptional<float> MaxWindowWidth, TOptional<float> MaxWindowHeight) = 0;
 
+		/**
+		 * should be called after the window is created and registered and before scene rendering begins
+		*/
+		virtual void PrepareDeviceDisplay() {};
+
 		/**
 		* Apply PieWindow device parameters
 		*/
diff --git a/Engine/Source/Runtime/Slate/Private/Framework/Application/SlateApplication.cpp b/Engine/Source/Runtime/Slate/Private/Framework/Application/SlateApplication.cpp
--- a/Engine/Source/Runtime/Slate/Private/Framework/Application/SlateApplication.cpp
+++ b/Engine/Source/Runtime/Slate/Private/Framework/Application/SlateApplication.cpp
@@ -2033,6 +2033,8 @@ TSharedRef< FGenericWindow > FSlateApplication::MakeWindow( TSharedRef<SWindow>
 
 	Definition->SizeLimits = InSlateWindow->GetSizeLimits();
 
+	Definition->bManualDPI = InSlateWindow->IsManualManageDPIChanges();
+
 	TSharedRef< FGenericWindow > NewWindow = PlatformApplication->MakeWindow();
 
 	if (LIKELY(FApp::CanEverRender()))
@@ -6627,6 +6629,18 @@ void FSlateApplication::FinishedReshapingWindow( const TSharedRef< FGenericWindo
 	}
 }
 
+void FSlateApplication::SignalSystemDPIChanged(const TSharedRef<FGenericWindow>& PlatformWindow)
+{
+#if WITH_EDITOR
+	TSharedPtr< SWindow > SlateWindow = FSlateWindowHelper::FindWindowByPlatformWindow(SlateWindows, PlatformWindow);
+
+	if (SlateWindow.IsValid() && SlateWindow->IsRegularWindow())
+	{
+		OnSignalSystemDPIChangedEvent.Broadcast(SlateWindow.ToSharedRef());
+	}
+#endif
+}
+
 void FSlateApplication::HandleDPIScaleChanged(const TSharedRef<FGenericWindow>& PlatformWindow)
 {
 #if WITH_EDITOR
@@ -6639,7 +6653,6 @@ void FSlateApplication::HandleDPIScaleChanged(const TSharedRef<FGenericWindow>&
 #endif
 }
 
-
 void FSlateApplication::OnMovedWindow( const TSharedRef< FGenericWindow >& PlatformWindow, const int32 X, const int32 Y )
 {
 	TSharedPtr< SWindow > Window = FSlateWindowHelper::FindWindowByPlatformWindow( SlateWindows, PlatformWindow );
diff --git a/Engine/Source/Runtime/Slate/Private/Widgets/Input/SVirtualJoystick.cpp b/Engine/Source/Runtime/Slate/Private/Widgets/Input/SVirtualJoystick.cpp
--- a/Engine/Source/Runtime/Slate/Private/Widgets/Input/SVirtualJoystick.cpp
+++ b/Engine/Source/Runtime/Slate/Private/Widgets/Input/SVirtualJoystick.cpp
@@ -339,6 +339,9 @@ void SVirtualJoystick::Tick( const FGeometry& AllottedGeometry, const double InC
 	// count how many controls are active
 	int32 NumActiveControls = 0;
 
+	// figure out how much to scale the control sizes
+	float ScaleFactor = GetScaleFactor(AllottedGeometry);
+
 	for (int32 ControlIndex = 0; ControlIndex < Controls.Num(); ControlIndex++)
 	{
 		FControlInfo& Control = Controls[ControlIndex];
@@ -362,11 +365,8 @@ void SVirtualJoystick::Tick( const FGeometry& AllottedGeometry, const double InC
 
 		// calculate absolute positions based on geometry
 		// @todo: Need to manage geometry changing!
-		if (!Control.bHasBeenPositioned)
+		if (!Control.bHasBeenPositioned || ScaleFactor != PreviousScalingFactor)
 		{
-			// figure out how much to scale the control sizes
-			float ScaleFactor = GetScaleFactor(AllottedGeometry);
-
 			// update all the sizes
 			Control.CorrectedCenter = FVector2D(ResolveRelativePosition(Control.Center.X, AllottedGeometry.GetLocalSize().X, ScaleFactor), ResolveRelativePosition(Control.Center.Y, AllottedGeometry.GetLocalSize().Y, ScaleFactor));
 			Control.VisualCenter = Control.CorrectedCenter;
@@ -403,7 +403,6 @@ void SVirtualJoystick::Tick( const FGeometry& AllottedGeometry, const double InC
 			const FGamepadKeyNames::Type XAxis = (Control.MainInputKey.IsValid() ? Control.MainInputKey.GetFName() : (ControlIndex == 0 ? FGamepadKeyNames::LeftAnalogX : FGamepadKeyNames::RightAnalogX));
 			const FGamepadKeyNames::Type YAxis = (Control.AltInputKey.IsValid() ? Control.AltInputKey.GetFName() : (ControlIndex == 0 ? FGamepadKeyNames::LeftAnalogY : FGamepadKeyNames::RightAnalogY));
 
-	//		UE_LOG(LogTemp, Log, TEXT("Joysticking %f,%f"), NormalizedOffset.X, -NormalizedOffset.Y);
 			FSlateApplication::Get().SetAllUserFocusToGameViewport();
 			FSlateApplication::Get().OnControllerAnalog(XAxis, 0, NormalizedOffset.X);
 			FSlateApplication::Get().OnControllerAnalog(YAxis, 0, -NormalizedOffset.Y);
@@ -416,6 +415,8 @@ void SVirtualJoystick::Tick( const FGeometry& AllottedGeometry, const double InC
 		}
 	}
 
+	// we need to store the computed scale factor so we can compare it with the value computed in the following frame and, if necessary, recompute widget position
+	PreviousScalingFactor = ScaleFactor;
 
 	// STATE MACHINE!
 	if (NumActiveControls > 0 || bPreventReCenter)
diff --git a/Engine/Source/Runtime/Slate/Public/Framework/Application/SlateApplication.h b/Engine/Source/Runtime/Slate/Public/Framework/Application/SlateApplication.h
--- a/Engine/Source/Runtime/Slate/Public/Framework/Application/SlateApplication.h
+++ b/Engine/Source/Runtime/Slate/Public/Framework/Application/SlateApplication.h
@@ -642,6 +642,9 @@ public:
 	/** Gets a delegate that is invoked in the editor when a windows dpi scale changes or when a widget window may have changed and DPI scale info needs to be checked */
 	DECLARE_EVENT_OneParam(FSlateApplication, FOnWindowDPIScaleChanged, TSharedRef<SWindow>);
 	FOnWindowDPIScaleChanged& OnWindowDPIScaleChanged() { return OnWindowDPIScaleChangedEvent; }
+
+	/** Event used to signal that a DPI change is about to happen */
+	FOnWindowDPIScaleChanged& OnSystemSignalsDPIChanged() { return OnSignalSystemDPIChangedEvent; }
 #endif //WITH_EDITOR
 
 	/**
@@ -1442,6 +1445,7 @@ public:
 	virtual void OnResizingWindow( const TSharedRef< FGenericWindow >& PlatformWindow ) override;
 	virtual bool BeginReshapingWindow( const TSharedRef< FGenericWindow >& PlatformWindow ) override;
 	virtual void FinishedReshapingWindow( const TSharedRef< FGenericWindow >& PlatformWindow ) override;
+	virtual void SignalSystemDPIChanged(const TSharedRef<FGenericWindow>& Window) override;
 	virtual void HandleDPIScaleChanged(const TSharedRef<FGenericWindow>& Window) override;
 	virtual void OnMovedWindow( const TSharedRef< FGenericWindow >& PlatformWindow, const int32 X, const int32 Y ) override;
 	virtual bool OnWindowActivationChanged( const TSharedRef< FGenericWindow >& PlatformWindow, const EWindowActivation ActivationType ) override;
@@ -2206,6 +2210,11 @@ private:
 	 */
 	FOnApplicationMousePreInputButtonDownListener OnApplicationMousePreInputButtonDownListenerEvent;
 
+	/**
+	* Called before the dpi scale of a particular window is about to changed
+	*/
+	FOnWindowDPIScaleChanged OnSignalSystemDPIChangedEvent;
+
 	/**
 	 * Called when an editor window dpi scale is changed
 	 */
diff --git a/Engine/Source/Runtime/Slate/Public/Widgets/Input/SVirtualJoystick.h b/Engine/Source/Runtime/Slate/Public/Widgets/Input/SVirtualJoystick.h
--- a/Engine/Source/Runtime/Slate/Public/Widgets/Input/SVirtualJoystick.h
+++ b/Engine/Source/Runtime/Slate/Public/Widgets/Input/SVirtualJoystick.h
@@ -218,4 +218,7 @@ protected:
 
 	/* Countdown until next state change */
 	float Countdown;
+
+	/** Last used scaling value for  */
+	float PreviousScalingFactor = 0.0f;
 };
diff --git a/Engine/Source/Runtime/SlateCore/Private/Widgets/SWindow.cpp b/Engine/Source/Runtime/SlateCore/Private/Widgets/SWindow.cpp
--- a/Engine/Source/Runtime/SlateCore/Private/Widgets/SWindow.cpp
+++ b/Engine/Source/Runtime/SlateCore/Private/Widgets/SWindow.cpp
@@ -249,6 +249,8 @@ void SWindow::Construct(const FArguments& InArgs)
 		.SetMinHeight(InArgs._MinHeight)
 		.SetMaxWidth(InArgs._MaxWidth)
 		.SetMaxHeight(InArgs._MaxHeight);
+	this->bManualManageDPI = InArgs._bManualManageDPI;
+
 	bCanTick = false;
 
 	// calculate window size from client size
@@ -1103,6 +1105,24 @@ float SWindow::GetDPIScaleFactor() const
 	return 1.0f;
 }
 
+void SWindow::SetDPIScaleFactor(const float Factor)
+{
+	if (NativeWindow.IsValid())
+	{
+		NativeWindow->SetDPIScaleFactor(Factor);
+	}
+}
+
+void SWindow::SetManualManageDPIChanges(const bool bManualDPI)
+{
+	bManualManageDPI = bManualDPI;
+
+	if (NativeWindow.IsValid())
+	{
+		NativeWindow->SetManualManageDPIChanges(bManualManageDPI);
+	}
+}
+
 bool SWindow::IsDescendantOf( const TSharedPtr<SWindow>& ParentWindow ) const
 {
 	TSharedPtr<SWindow> CandidateToCheck = this->GetParentWindow();
@@ -1882,6 +1902,7 @@ SWindow::SWindow()
 	, ExpectedMaxHeight( INDEX_NONE )
 	, TitleBar()
 	, bIsDrawingEnabled( true )
+	, bManualManageDPI( false )
 {
 }
 
diff --git a/Engine/Source/Runtime/SlateCore/Public/Widgets/SWindow.h b/Engine/Source/Runtime/SlateCore/Public/Widgets/SWindow.h
--- a/Engine/Source/Runtime/SlateCore/Public/Widgets/SWindow.h
+++ b/Engine/Source/Runtime/SlateCore/Public/Widgets/SWindow.h
@@ -260,6 +260,9 @@ public:
 		/** The margin around the edges of the window that will be detected as places the user can grab to resize the window. */
 		SLATE_ARGUMENT(FMargin, UserResizeBorder)
 
+		/** true if this window will self handle any eventual DPI adjustments */
+		SLATE_ARGUMENT(bool, bManualManageDPI)
+
 		SLATE_DEFAULT_SLOT( FArguments, Content )
 
 	SLATE_END_ARGS()
@@ -447,6 +450,17 @@ public:
 	/** Returns the DPI scale factor of the native window */
 	float GetDPIScaleFactor() const;
 
+	/** Overrides the DPI scale factor of the native window */
+	void SetDPIScaleFactor(const float Factor);
+
+	/** Will inform the native window that we need to handle any DPI changes from within the application */
+	void SetManualManageDPIChanges(const bool bManualDPI);
+
+	bool IsManualManageDPIChanges() const
+	{
+		return bManualManageDPI;
+	}
+
 	/** 
 	 * Returns whether or not this window is a descendant of the specfied parent window
 	 *
@@ -980,6 +994,8 @@ protected:
 	/** True if the window should preserve its aspect ratio when resized by user */
 	bool bShouldPreserveAspectRatio : 1;
 
+	bool bManualManageDPI : 1;
+
 	/** When should the window be activated upon being shown */
 	EWindowActivationPolicy WindowActivationPolicy;
 
