UE-63378 - Fixed static analysis warnings

diff --git a/Engine/Source/Editor/PIEPreviewDeviceProfileSelector/Private/PIEPreviewDevice.cpp b/Engine/Source/Editor/PIEPreviewDeviceProfileSelector/Private/PIEPreviewDevice.cpp
--- a/Engine/Source/Editor/PIEPreviewDeviceProfileSelector/Private/PIEPreviewDevice.cpp
+++ b/Engine/Source/Editor/PIEPreviewDeviceProfileSelector/Private/PIEPreviewDevice.cpp
@@ -358,14 +358,17 @@ void FPIEPreviewDevice::ApplyRHIOverrides() const
 		UMaterialShaderQualitySettings::Get()->SetPreviewPlatform(QualityPreviewShaderPlatform);
 	}
 
-	GMaxTextureDimensions.SetPreviewOverride(RHIOverrideState->MaxTextureDimensions);
-	GMaxShadowDepthBufferSizeX.SetPreviewOverride(RHIOverrideState->MaxShadowDepthBufferSizeX);
-	GMaxShadowDepthBufferSizeY.SetPreviewOverride(RHIOverrideState->MaxShadowDepthBufferSizeY);
-	GMaxCubeTextureDimensions.SetPreviewOverride(RHIOverrideState->MaxCubeTextureDimensions);
-	GRHISupportsInstancing.SetPreviewOverride(RHIOverrideState->SupportsInstancing);
-	GSupportsMultipleRenderTargets.SetPreviewOverride(RHIOverrideState->SupportsMultipleRenderTargets);
-	GSupportsRenderTargetFormat_PF_FloatRGBA.SetPreviewOverride(RHIOverrideState->SupportsRenderTargetFormat_PF_FloatRGBA);
-	GSupportsRenderTargetFormat_PF_G8.SetPreviewOverride(RHIOverrideState->SupportsRenderTargetFormat_PF_G8);
+	if (RHIOverrideState != nullptr)
+	{
+		GMaxTextureDimensions.SetPreviewOverride(RHIOverrideState->MaxTextureDimensions);
+		GMaxShadowDepthBufferSizeX.SetPreviewOverride(RHIOverrideState->MaxShadowDepthBufferSizeX);
+		GMaxShadowDepthBufferSizeY.SetPreviewOverride(RHIOverrideState->MaxShadowDepthBufferSizeY);
+		GMaxCubeTextureDimensions.SetPreviewOverride(RHIOverrideState->MaxCubeTextureDimensions);
+		GRHISupportsInstancing.SetPreviewOverride(RHIOverrideState->SupportsInstancing);
+		GSupportsMultipleRenderTargets.SetPreviewOverride(RHIOverrideState->SupportsMultipleRenderTargets);
+		GSupportsRenderTargetFormat_PF_FloatRGBA.SetPreviewOverride(RHIOverrideState->SupportsRenderTargetFormat_PF_FloatRGBA);
+		GSupportsRenderTargetFormat_PF_G8.SetPreviewOverride(RHIOverrideState->SupportsRenderTargetFormat_PF_G8);
+	}
 }
 
 FString FPIEPreviewDevice::GetProfile() const
