Mobile PIE bugfixes

diff --git a/Engine/Source/Editor/PIEPreviewDeviceProfileSelector/Private/PIEPreviewDevice.cpp b/Engine/Source/Editor/PIEPreviewDeviceProfileSelector/Private/PIEPreviewDevice.cpp
--- a/Engine/Source/Editor/PIEPreviewDeviceProfileSelector/Private/PIEPreviewDevice.cpp
+++ b/Engine/Source/Editor/PIEPreviewDeviceProfileSelector/Private/PIEPreviewDevice.cpp
@@ -15,6 +15,7 @@
 #include "ImageUtils.h"
 #include "HAL/PlatformApplicationMisc.h"
 #include "SEditorViewport.h"
+#include "Widgets/SWindow.h"
 
 static const int32 g_JSON_VALUE_NOT_SET = 0;
 
@@ -94,6 +95,8 @@ void FPIEPreviewDevice::ComputeViewportSize(const bool bClampWindowSize)
 
 			WindowWidth = DesktopWidth;
 			WindowHeight *= ScaleFactor;
+
+			ViewportMargin = ViewportMargin * ScaleFactor;
 		}
 		if (WindowHeight > DesktopHeight)
 		{
@@ -103,6 +106,8 @@ void FPIEPreviewDevice::ComputeViewportSize(const bool bClampWindowSize)
 
 			WindowWidth *= ScaleFactor;
 			WindowHeight = DesktopHeight;
+
+			ViewportMargin = ViewportMargin * ScaleFactor;
 		}
 	}
 }
@@ -355,7 +360,7 @@ void FPIEPreviewDevice::ApplyRHIOverrides() const
 	if (PreviewPlatform != SP_NumPlatforms)
 	{
 		FName QualityPreviewShaderPlatform = LegacyShaderPlatformToShaderFormat(PreviewPlatform);
-		UMaterialShaderQualitySettings::Get()->SetPreviewPlatform(QualityPreviewShaderPlatform);
+		//UMaterialShaderQualitySettings::Get()->SetPreviewPlatform(QualityPreviewShaderPlatform);
 	}
 
 	if (RHIOverrideState != nullptr)
@@ -415,4 +420,9 @@ FString FPIEPreviewDevice::GetProfile() const
 	}
 
 	return Profile;
+}
+
+int32 FPIEPreviewDevice::GetWindowClientHeight() const
+{
+	return WindowHeight - WindowTitleBarSize;
 }
\ No newline at end of file
diff --git a/Engine/Source/Editor/PIEPreviewDeviceProfileSelector/Private/PIEPreviewDevice.h b/Engine/Source/Editor/PIEPreviewDeviceProfileSelector/Private/PIEPreviewDevice.h
--- a/Engine/Source/Editor/PIEPreviewDeviceProfileSelector/Private/PIEPreviewDevice.h
+++ b/Engine/Source/Editor/PIEPreviewDeviceProfileSelector/Private/PIEPreviewDevice.h
@@ -80,6 +80,10 @@ public:
 	int32 GetWindowWidth() const;
 	int32 GetWindowHeight() const;
 
+	// functions to retrieve the window client area size
+	int32 GetWindowClientWidth() const;
+	int32 GetWindowClientHeight() const;
+
 	// returns true if the current device was rotated from the original orientation
 	bool IsDeviceFlipped() const;
 
@@ -100,6 +104,7 @@ public:
 	// if true the device supports rotations
 	bool IsRotationAllowed() const;
 
+	// call this to enable or disable mobile content scale factor effects
 	void SetIgnoreMobileContentScaleFactor(bool bIgnore);
 
 	FString GetProfile() const;
@@ -127,6 +132,11 @@ FORCEINLINE int32 FPIEPreviewDevice::GetWindowHeight() const
 	return WindowHeight;
 }
 
+FORCEINLINE int32 FPIEPreviewDevice::GetWindowClientWidth() const
+{
+	return WindowWidth;
+}
+
 FORCEINLINE bool FPIEPreviewDevice::IsDeviceFlipped() const
 {
 	return bDeviceFlipped;
diff --git a/Engine/Source/Editor/PIEPreviewDeviceProfileSelector/Private/PIEPreviewDeviceProfileSelectorModule.cpp b/Engine/Source/Editor/PIEPreviewDeviceProfileSelector/Private/PIEPreviewDeviceProfileSelectorModule.cpp
--- a/Engine/Source/Editor/PIEPreviewDeviceProfileSelector/Private/PIEPreviewDeviceProfileSelectorModule.cpp
+++ b/Engine/Source/Editor/PIEPreviewDeviceProfileSelector/Private/PIEPreviewDeviceProfileSelectorModule.cpp
@@ -55,6 +55,21 @@ void FPIEPreviewDeviceModule::InitPreviewDevice()
 
 TSharedRef<SWindow> FPIEPreviewDeviceModule::CreatePIEPreviewDeviceWindow(FVector2D ClientSize, FText WindowTitle, EAutoCenter AutoCenterType, FVector2D ScreenPosition, TOptional<float> MaxWindowWidth, TOptional<float> MaxWindowHeight)
 {
+	if (ScreenPosition.IsNearlyZero())
+	{
+		int32 WinX, WinY;
+		bool bFoundX = GConfig->GetInt(TEXT("/Script/Engine.MobilePIE"), TEXT("WindowPosX"), WinX, GEngineIni);
+		bool bFoundY = GConfig->GetInt(TEXT("/Script/Engine.MobilePIE"), TEXT("WindowPosY"), WinY, GEngineIni);
+
+		if (bFoundX && bFoundY)
+		{
+			ScreenPosition.X = WinX;
+			ScreenPosition.Y = WinY;
+
+			AutoCenterType = EAutoCenter::None;
+		}
+	}
+
 	FPIEPreviewWindowCoreStyle::InitializePIECoreStyle();
 
 	static FWindowStyle BackgroundlessStyle = FCoreStyle::Get().GetWidgetStyle<FWindowStyle>("Window");
@@ -92,10 +107,10 @@ void FPIEPreviewDeviceModule::UpdateDisplayResolution()
 		return;
 	}
 
-	const int32 WindowWidth = Device->GetWindowWidth();
-	const int32 WindowHeight = Device->GetWindowHeight();
+	const int32 ClientWidth = Device->GetWindowClientWidth();
+	const int32 ClientHeight = Device->GetWindowClientHeight();
 
-	FSystemResolution::RequestResolutionChange(WindowWidth, WindowHeight, EWindowMode::Windowed);
+	FSystemResolution::RequestResolutionChange(ClientWidth, ClientHeight, EWindowMode::Windowed);
 	IConsoleManager::Get().CallAllConsoleVariableSinks();
 }
 
diff --git a/Engine/Source/Editor/PIEPreviewDeviceProfileSelector/Private/PIEPreviewWindow.cpp b/Engine/Source/Editor/PIEPreviewDeviceProfileSelector/Private/PIEPreviewWindow.cpp
--- a/Engine/Source/Editor/PIEPreviewDeviceProfileSelector/Private/PIEPreviewWindow.cpp
+++ b/Engine/Source/Editor/PIEPreviewDeviceProfileSelector/Private/PIEPreviewWindow.cpp
@@ -13,6 +13,7 @@
 #include "Widgets/Input/SMenuAnchor.h"
 #include "Framework/MultiBox/MultiBoxBuilder.h"
 #include "Widgets/Input/SCheckBox.h"
+#include "Misc/ConfigCacheIni.h"
 
 
 #if WITH_EDITOR
@@ -364,6 +365,10 @@ void SPIEPreviewWindow::OnWindowMoved(const TSharedRef<SWindow>& Window)
 	{
 		ScaleDeviceToPhisicalSize();
 	}
+
+	FVector2D WindowPos = GetPositionInScreen();
+	GConfig->SetInt(TEXT("/Script/Engine.MobilePIE"), TEXT("WindowPosX"), FMath::CeilToInt(WindowPos.X), GEngineIni);
+	GConfig->SetInt(TEXT("/Script/Engine.MobilePIE"), TEXT("WindowPosY"), FMath::CeilToInt(WindowPos.Y), GEngineIni);
 }
 
 void SPIEPreviewWindow::OnDisplayDPIChanged(TSharedRef<SWindow> Window)
diff --git a/Engine/Source/Editor/PIEPreviewDeviceProfileSelector/Private/PIEPreviewWindow.h b/Engine/Source/Editor/PIEPreviewDeviceProfileSelector/Private/PIEPreviewWindow.h
--- a/Engine/Source/Editor/PIEPreviewDeviceProfileSelector/Private/PIEPreviewWindow.h
+++ b/Engine/Source/Editor/PIEPreviewDeviceProfileSelector/Private/PIEPreviewWindow.h
@@ -85,10 +85,10 @@ private:
 	// utility function that will compute a screen and DPI scale factor needed to scale the display window to the physical device size
 	void ComputeScaleToDeviceSizeFactor(float& OutScreenFactor, float& OutDPIScaleFactor) const;
 
-	// callback when DPI value varies, useful when we need to constrain the window to the physical device size
+	// callback used when DPI value varies, useful when we need to constrain the window to the physical device size
 	void OnWindowMoved(const TSharedRef<SWindow>& Window);
 
-	// callback when DPI value varies, useful when we need to constrain the window to the physical device size
+	// callback used when DPI value varies, useful when we need to constrain the window to the physical device size
 	void OnDisplayDPIChanged(TSharedRef<SWindow> Window);
 
 protected:
