[bugfix] Accidentally disabled SetPreviewPlatform() call for mobile PIE simulation.

diff --git a/Engine/Source/Editor/PIEPreviewDeviceProfileSelector/Private/PIEPreviewDevice.cpp b/Engine/Source/Editor/PIEPreviewDeviceProfileSelector/Private/PIEPreviewDevice.cpp
--- a/Engine/Source/Editor/PIEPreviewDeviceProfileSelector/Private/PIEPreviewDevice.cpp
+++ b/Engine/Source/Editor/PIEPreviewDeviceProfileSelector/Private/PIEPreviewDevice.cpp
@@ -360,7 +360,7 @@ void FPIEPreviewDevice::ApplyRHIOverrides() const
 	if (PreviewPlatform != SP_NumPlatforms)
 	{
 		FName QualityPreviewShaderPlatform = LegacyShaderPlatformToShaderFormat(PreviewPlatform);
-		//UMaterialShaderQualitySettings::Get()->SetPreviewPlatform(QualityPreviewShaderPlatform);
+		UMaterialShaderQualitySettings::Get()->SetPreviewPlatform(QualityPreviewShaderPlatform);
 	}
 
 	if (RHIOverrideState != nullptr)
