Mobile PIE improvements

diff --git a/Engine/Source/Editor/PIEPreviewDeviceProfileSelector/Private/PIEPreviewDevice.cpp b/Engine/Source/Editor/PIEPreviewDeviceProfileSelector/Private/PIEPreviewDevice.cpp
--- a/Engine/Source/Editor/PIEPreviewDeviceProfileSelector/Private/PIEPreviewDevice.cpp
+++ b/Engine/Source/Editor/PIEPreviewDeviceProfileSelector/Private/PIEPreviewDevice.cpp
@@ -24,6 +24,14 @@ FPIEPreviewDevice::FPIEPreviewDevice()
 	DeviceSpecs = MakeShareable(new FPIEPreviewDeviceSpecifications());
 }
 
+void FPIEPreviewDevice::ShutdownDevice()
+{
+	if (BezelTexture != nullptr && BezelTexture->IsValidLowLevel())
+	{
+		BezelTexture->RemoveFromRoot();
+	}
+}
+
 void FPIEPreviewDevice::ComputeViewportSize(const bool bClampWindowSize)
 {
 	int32 ScreenWidth, ScreenHeight;
@@ -124,11 +132,11 @@ void FPIEPreviewDevice::GetDeviceDefaultResolution(int32& Width, int32& Height)
 	Height = DeviceSpecs->ResolutionY;
 }
 
-void FPIEPreviewDevice::ComputeDeviceResolution(int32& Width, int32& Height)
+void FPIEPreviewDevice::ComputeContentScaledResolution(int32& Width, int32& Height)
 {
 	GetDeviceDefaultResolution(Width, Height);
 
-	if (!bIgnoreContentScaleFactor)
+	if (!GetIgnoreMobileContentScaleFactor())
 	{
 		switch (DeviceSpecs->DevicePlatform)
 		{
@@ -160,9 +168,14 @@ void FPIEPreviewDevice::ComputeDeviceResolution(int32& Width, int32& Height)
 			break;
 
 			default:
-				break;
+			break;
 		} //end switch
 	}// end if (!bIgnoreContentScaleFactor)
+}
+
+void FPIEPreviewDevice::ComputeDeviceResolution(int32& Width, int32& Height)
+{
+	ComputeContentScaledResolution(Width, Height);
 
 	Width *= ResolutionScaleFactor;
 	Height *= ResolutionScaleFactor;
@@ -268,7 +281,7 @@ void FPIEPreviewDevice::SetupDevice(const int32 InWindowTitleBarSize)
 	WindowTitleBarSize = InWindowTitleBarSize;
 
 	// set initial scale factor
-	ResolutionScaleFactor = 0.5f;
+	ResolutionScaleFactor = 1.0f;
 
 	// compute bezel file path
 	FString BezelPath = FPaths::EngineContentDir() + TEXT("Editor/PIEPreviewDeviceSpecs/");
@@ -284,6 +297,7 @@ void FPIEPreviewDevice::SetupDevice(const int32 InWindowTitleBarSize)
 
 	// load the bezel texture
 	BezelTexture = FImageUtils::ImportFileAsTexture2D(BezelPath);
+	BezelTexture->AddToRoot();
 	
 	// if we have invalid/uninitialized viewport values use the values provided as native device resolution
 	FPIEPreviewDeviceBezelViewportRect& ViewportRect = DeviceSpecs->BezelProperties.BezelViewportRect;
@@ -304,18 +318,11 @@ void FPIEPreviewDevice::SetupDevice(const int32 InWindowTitleBarSize)
 	// determine current orientation
 	bool bSwitchOrientation = bLandscape && DeviceSpecs->ResolutionY > DeviceSpecs->ResolutionX;
 	bSwitchOrientation |= !bLandscape && bPortrait && DeviceSpecs->ResolutionX > DeviceSpecs->ResolutionY;
+
 	if (bSwitchOrientation)
 	{
 		SwitchOrientation(true);
 	}
-	else
-	{
-		// finally compute needed window and viewport size
-		// the above branch will call it from SwitchOrientation()
-		ComputeViewportSize(true);
-	}
-
-	ApplyRHIOverrides();
 }
 
 void FPIEPreviewDevice::ApplyRHIPrerequisitesOverrides() const
@@ -421,5 +428,5 @@ FString FPIEPreviewDevice::GetProfile() const
 
 int32 FPIEPreviewDevice::GetWindowClientHeight() const
 {
-	return WindowHeight - WindowTitleBarSize;
+	return WindowHeight - WindowTitleBarSize * DPIScaleFactor;
 }
\ No newline at end of file
diff --git a/Engine/Source/Editor/PIEPreviewDeviceProfileSelector/Private/PIEPreviewDevice.h b/Engine/Source/Editor/PIEPreviewDeviceProfileSelector/Private/PIEPreviewDevice.h
--- a/Engine/Source/Editor/PIEPreviewDeviceProfileSelector/Private/PIEPreviewDevice.h
+++ b/Engine/Source/Editor/PIEPreviewDeviceProfileSelector/Private/PIEPreviewDevice.h
@@ -43,13 +43,9 @@ private:
 	/** when true r.MobileContentScaleFactor is ignored */
 	bool bIgnoreContentScaleFactor = false;
 
-	UPROPERTY()
 	class UTexture2D* BezelTexture = nullptr;
 
 private:
-	// utility function that computes the viewport resolution
-	void ComputeDeviceResolution(int32& ScreenWidth, int32& ScreenHeight);
-
 	// utility function that will attempt to determine the supported device orientations
 	void DetermineScreenOrientationRequirements(bool& bNeedPortrait, bool& bNeedLandscape);
 
@@ -57,6 +53,7 @@ private:
 	ERHIFeatureLevel::Type GetPreviewDeviceFeatureLevel() const;
 
 	// call this function to apply specific RHI settings as specified in the json file
+public:
 	void ApplyRHIOverrides() const;
 
 	// function that computes the viewport widget offset and needed window size
@@ -68,6 +65,8 @@ public:
 	// call this to run device setup -> load the bezel texture, compute appropriate orientation and apply device specific RHI settings
 	void SetupDevice(const int32 InWindowTitleBarSize);
 
+	void ShutdownDevice();
+
 	// call this before RHI creation to apply needed setup overrides
 	void ApplyRHIPrerequisitesOverrides() const;
 
@@ -107,6 +106,15 @@ public:
 	// call this to enable or disable mobile content scale factor effects
 	void SetIgnoreMobileContentScaleFactor(bool bIgnore);
 
+	// call to check if the mobile content scale factor is taken into account
+	bool GetIgnoreMobileContentScaleFactor() const;
+
+	// utility function that computes the viewport resolution
+	void ComputeDeviceResolution(int32& ScreenWidth, int32& ScreenHeight);
+
+	// utility function that computes the resolution after applying r.MobileContentScaleFactor
+	void ComputeContentScaledResolution(int32& ScreenWidth, int32& ScreenHeight);
+
 	FString GetProfile() const;
 
 	UTexture2D *GetBezelTexture();
@@ -187,4 +195,9 @@ FORCEINLINE UTexture2D *FPIEPreviewDevice::GetBezelTexture()
 FORCEINLINE void FPIEPreviewDevice::SetIgnoreMobileContentScaleFactor(bool bIgnore)
 {
 	bIgnoreContentScaleFactor = bIgnore;
+}
+
+FORCEINLINE bool FPIEPreviewDevice::GetIgnoreMobileContentScaleFactor() const
+{
+	return bIgnoreContentScaleFactor;
 }
\ No newline at end of file
diff --git a/Engine/Source/Editor/PIEPreviewDeviceProfileSelector/Private/PIEPreviewDeviceProfileSelectorModule.cpp b/Engine/Source/Editor/PIEPreviewDeviceProfileSelector/Private/PIEPreviewDeviceProfileSelectorModule.cpp
--- a/Engine/Source/Editor/PIEPreviewDeviceProfileSelector/Private/PIEPreviewDeviceProfileSelectorModule.cpp
+++ b/Engine/Source/Editor/PIEPreviewDeviceProfileSelector/Private/PIEPreviewDeviceProfileSelectorModule.cpp
@@ -19,6 +19,8 @@
 #include "PIEPreviewWindow.h"
 #include "Framework/Application/SlateApplication.h"
 #include "UnrealEngine.h"
+#include "Engine/GameViewportClient.h"
+#include "Engine/UserInterfaceSettings.h"
 
 DECLARE_LOG_CATEGORY_EXTERN(LogPIEPreviewDevice, Log, All); 
 DEFINE_LOG_CATEGORY(LogPIEPreviewDevice);
@@ -26,10 +28,36 @@ IMPLEMENT_MODULE(FPIEPreviewDeviceModule, PIEPreviewDeviceProfileSelector);
 
 void FPIEPreviewDeviceModule::StartupModule()
 {
+	// the window size will be available after all data is loaded and we'll use this callback to display it
+	EngineInitCompleteDelegate = FCoreDelegates::OnFEngineLoopInitComplete.AddRaw(this, &FPIEPreviewDeviceModule::OnEngineInitComplete);
+
+	// to finish setup we need complete engine initialization
+	ViewportCreatedDelegate = UGameViewportClient::OnViewportCreated().AddRaw(this, &FPIEPreviewDeviceModule::OnViewportCreated);
 }
 
 void FPIEPreviewDeviceModule::ShutdownModule()
 {
+	// clear delegates set in StartupModule()
+	if (EngineInitCompleteDelegate.IsValid())
+	{
+		FCoreDelegates::OnFEngineLoopInitComplete.Remove(EngineInitCompleteDelegate);
+	}
+
+	if (ViewportCreatedDelegate.IsValid())
+	{
+		UGameViewportClient::OnViewportCreated().Remove(ViewportCreatedDelegate);
+	}
+
+	TSharedPtr<SPIEPreviewWindow> WindowPtr = WindowWPtr.Pin();
+	if (WindowPtr.IsValid())
+	{
+		WindowPtr->PrepareShutdown();
+	}
+
+	if (Device.IsValid())
+	{
+		Device->ShutdownDevice();
+	}
 }
 
 FString const FPIEPreviewDeviceModule::GetRuntimeDeviceProfileName()
@@ -53,28 +81,57 @@ void FPIEPreviewDeviceModule::InitPreviewDevice()
 	}
 }
 
-TSharedRef<SWindow> FPIEPreviewDeviceModule::CreatePIEPreviewDeviceWindow(FVector2D ClientSize, FText WindowTitle, EAutoCenter AutoCenterType, FVector2D ScreenPosition, TOptional<float> MaxWindowWidth, TOptional<float> MaxWindowHeight)
+void FPIEPreviewDeviceModule::OnEngineInitComplete()
 {
-	if (ScreenPosition.IsNearlyZero())
+	TSharedPtr<SPIEPreviewWindow> WindowPtr = WindowWPtr.Pin();
+
+	// TODO: Localization
+	FString AppTitle = FGlobalTabmanager::Get()->GetApplicationTitle().ToString() + "Previewing: " + PreviewDevice;
+	FGlobalTabmanager::Get()->SetApplicationTitle(FText::FromString(AppTitle));
+
+	if (WindowPtr.IsValid())
 	{
-		int32 WinX, WinY;
-		bool bFoundX = GConfig->GetInt(TEXT("/Script/Engine.MobilePIE"), TEXT("WindowPosX"), WinX, GEngineIni);
-		bool bFoundY = GConfig->GetInt(TEXT("/Script/Engine.MobilePIE"), TEXT("WindowPosY"), WinY, GEngineIni);
+		int32 TitleBarSize = SPIEPreviewWindow::GetDefaultTitleBarSize();
+		Device->SetupDevice(TitleBarSize);
+		
+		WindowPtr->PrepareWindow(InitialWindowPosition, InitialWindowScaleValue, Device);
+		WindowPtr->ShowWindow();
+	}
+}
 
-		if (bFoundX && bFoundY)
-		{
-			ScreenPosition.X = WinX;
-			ScreenPosition.Y = WinY;
+bool FPIEPreviewDeviceModule::ReadWindowConfig()
+{
+	InitialWindowScaleValue = 0.0f;
+	GConfig->GetFloat(TEXT("/Script/Engine.MobilePIE"), TEXT("DeviceScalingFactor"), InitialWindowScaleValue, GEngineIni);
 
-			AutoCenterType = EAutoCenter::None;
-		}
+	// read window position
+	int32 WinX = 0, WinY = 0;
+	bool bFound = GConfig->GetInt(TEXT("/Script/Engine.MobilePIE"), TEXT("WindowPosX"), WinX, GEngineIni);
+	bFound &= GConfig->GetInt(TEXT("/Script/Engine.MobilePIE"), TEXT("WindowPosY"), WinY, GEngineIni);
+	InitialWindowPosition.Set(WinX, WinY);
+
+	return bFound;
+}
+
+TSharedRef<SWindow> FPIEPreviewDeviceModule::CreatePIEPreviewDeviceWindow(FVector2D ClientSize, FText WindowTitle, EAutoCenter AutoCenterType, FVector2D ScreenPosition, TOptional<float> MaxWindowWidth, TOptional<float> MaxWindowHeight)
+{
+	bool bPosFound = ReadWindowConfig();
+
+	if (ScreenPosition.IsNearlyZero() && bPosFound)
+	{
+		ScreenPosition = InitialWindowPosition;
+		AutoCenterType = EAutoCenter::None;
+	}
+	else
+	{
+		InitialWindowPosition = ScreenPosition;
 	}
 
 	FPIEPreviewWindowCoreStyle::InitializePIECoreStyle();
 
 	static FWindowStyle BackgroundlessStyle = FCoreStyle::Get().GetWidgetStyle<FWindowStyle>("Window");
 	BackgroundlessStyle.SetBackgroundBrush(FSlateNoResource());
-	TSharedRef<SPIEPreviewWindow> Window = SNew(SPIEPreviewWindow, Device)
+	TSharedRef<SPIEPreviewWindow> Window = SNew(SPIEPreviewWindow)
 		.Type(EWindowType::GameWindow)
 		.Style(&BackgroundlessStyle)
 		.ClientSize(ClientSize)
@@ -97,37 +154,41 @@ TSharedRef<SWindow> FPIEPreviewDeviceModule::CreatePIEPreviewDeviceWindow(FVecto
 
  	WindowWPtr = Window;
 
+	if (GameLayerManagerWidget.IsValid())
+	{
+		Window->SetGameLayerManagerWidget(GameLayerManagerWidget);
+	}
+
 	return Window;
 }
 
 void FPIEPreviewDeviceModule::UpdateDisplayResolution()
 {
-	if (!Device.IsValid())
+	TSharedPtr<SPIEPreviewWindow> WindowPtr = WindowWPtr.Pin();
+
+	if (!Device.IsValid() || !WindowPtr.IsValid())
 	{
 		return;
 	}
 
-	const int32 ClientWidth = Device->GetWindowClientWidth();
-	const int32 ClientHeight = Device->GetWindowClientHeight();
+	const int32 ClientWidth = Device->GetWindowWidth();
+	const int32 ClientHeight = Device->GetWindowHeight() - WindowPtr->GetTitleBarSize().Get();
 
 	FSystemResolution::RequestResolutionChange(ClientWidth, ClientHeight, EWindowMode::Windowed);
 	IConsoleManager::Get().CallAllConsoleVariableSinks();
 }
 
-void FPIEPreviewDeviceModule::PrepareDeviceDisplay()
+void FPIEPreviewDeviceModule::OnWindowReady(TSharedRef<SWindow> Window)
 {
-	TSharedPtr<SPIEPreviewWindow> WindowPtr = WindowWPtr.Pin();
+	TSharedPtr<SPIEPreviewWindow> WindowPtr = StaticCastSharedRef<SPIEPreviewWindow>(Window);
 
-	if (!WindowPtr.IsValid() && !Device.IsValid())
+	if (WindowPtr.IsValid())
 	{
-		return;
+		// the window will only be displayed after the loading is complete (OnEngineInitComplete)
+		WindowPtr->HideWindow();
 	}
 
 	FSlateApplication::Get().SetGameIsFakingTouchEvents(true);
-
-	WindowPtr->SetScaleWindowToDeviceSize(true);
-
-	UpdateDisplayResolution();
 }
 
 void FPIEPreviewDeviceModule::ApplyPreviewDeviceState()
@@ -137,15 +198,17 @@ void FPIEPreviewDeviceModule::ApplyPreviewDeviceState()
 		return;
 	}
 
-	// TODO: Localization
-	FString AppTitle = FGlobalTabmanager::Get()->GetApplicationTitle().ToString() + "Previewing: "+ PreviewDevice;
-	FGlobalTabmanager::Get()->SetApplicationTitle(FText::FromString(AppTitle));
-
-	int32 TitleBarSize = SPIEPreviewWindow::GetDefaultTitleBarSize();
-	Device->SetupDevice(TitleBarSize);
+	Device->ApplyRHIOverrides();
+}
 
-	// need to call this before the actual window is created in order override the window mode to EWindowMode::Windowed
-	UpdateDisplayResolution();
+void FPIEPreviewDeviceModule::OnViewportCreated()
+{
+	// disable mouse viewport locking
+	if (GEngine->GameViewport != nullptr)
+	{
+		GEngine->GameViewport->SetCaptureMouseOnClick(EMouseCaptureMode::NoCapture);
+		GEngine->GameViewport->SetMouseLockMode(EMouseLockMode::DoNotLock);
+	}
 }
 
 const FPIEPreviewDeviceContainer& FPIEPreviewDeviceModule::GetPreviewDeviceContainer()
@@ -177,7 +240,6 @@ FString FPIEPreviewDeviceModule::FindDeviceSpecificationFilePath(const FString&
 		}
 	}
 	return FoundPath;
-
 }
 
 bool FPIEPreviewDeviceModule::ReadDeviceSpecification()
@@ -217,3 +279,14 @@ bool FPIEPreviewDeviceModule::ReadDeviceSpecification()
 
 	return bValidDeviceSpec;
 }
+
+void FPIEPreviewDeviceModule::SetGameLayerManagerWidget(TSharedPtr<class SGameLayerManager> GameLayerManager)
+{
+	GameLayerManagerWidget = GameLayerManager;
+
+	TSharedPtr<SPIEPreviewWindow> WindowPtr = WindowWPtr.Pin();
+	if (WindowPtr.IsValid())
+	{
+		WindowPtr->SetGameLayerManagerWidget(GameLayerManager);
+	}
+}
\ No newline at end of file
diff --git a/Engine/Source/Editor/PIEPreviewDeviceProfileSelector/Private/PIEPreviewWindow.cpp b/Engine/Source/Editor/PIEPreviewDeviceProfileSelector/Private/PIEPreviewWindow.cpp
--- a/Engine/Source/Editor/PIEPreviewDeviceProfileSelector/Private/PIEPreviewWindow.cpp
+++ b/Engine/Source/Editor/PIEPreviewDeviceProfileSelector/Private/PIEPreviewWindow.cpp
@@ -14,118 +14,29 @@
 #include "Framework/MultiBox/MultiBoxBuilder.h"
 #include "Widgets/Input/SCheckBox.h"
 #include "Misc/ConfigCacheIni.h"
+#include "Slate/SGameLayerManager.h"
+#include "UnrealEngine.h"
 
 
 #if WITH_EDITOR
-
-//***********************************************************************************
-//SPIEToolbar Implementation
-//***********************************************************************************
-
-/** Toolbar class used to add some menus to configure various device display settings */
-class SPIEToolbar : public SViewportToolBar
-{
-public:
-	SLATE_BEGIN_ARGS(SPIEToolbar) {}
-		SLATE_EVENT(FOnGetContent, OnGetMenuContent)
-	SLATE_END_ARGS()
-
-	void Construct(const FArguments& InArgs);
-	FReply OnMenuClicked();
-
-protected:
-	TSharedPtr<SMenuAnchor> MenuAnchor;
-};
-
-void SPIEToolbar::Construct(const FArguments& InArgs)
-{
-	SViewportToolBar::Construct(SViewportToolBar::FArguments());
-
-	const FSlateBrush* ImageBrush = FPIEPreviewWindowCoreStyle::Get().GetBrush("ComboButton.Arrow");
-
-	const float MenuIconSize = 16.0f;
-
-	TSharedPtr<SWidget> ButtonContent =
-		SNew(SHorizontalBox)
-		+ SHorizontalBox::Slot()
-		.AutoWidth()
-		.Padding(2.0f, 2.0f)
-		[
-			SNew(SHorizontalBox)
-			+ SHorizontalBox::Slot()
-			.AutoWidth()
-			[
-				SNew(SBox)
-				.WidthOverride(MenuIconSize)
-				.HeightOverride(MenuIconSize)
-				.VAlign(VAlign_Center)
-				.HAlign(HAlign_Center)
-				[
-					SNew(SImage)
-					.Image(ImageBrush)
-					.ColorAndOpacity(FSlateColor::UseForeground())
-				]
-			]
-		];
-
-	ChildSlot
-	[
-		SAssignNew(MenuAnchor, SMenuAnchor)
-		.Placement(MenuPlacement_BelowAnchor)
-		[
-			SNew(SButton)
-			// Allows users to drag with the mouse to select options after opening the menu */
-			.ClickMethod(EButtonClickMethod::MouseDown)
-			.ContentPadding(FMargin(2.0f, 2.0f))
-			.VAlign(VAlign_Center)
-			.ButtonStyle(FPIEPreviewWindowCoreStyle::Get(), "PIEWindow.MenuButton")
-			.OnClicked(this, &SPIEToolbar::OnMenuClicked)
-			[
-				ButtonContent.ToSharedRef()
-			]
-		]
-		.OnGetMenuContent(InArgs._OnGetMenuContent)
-	];
-}
-
-FReply SPIEToolbar::OnMenuClicked()
-{
-	// If the menu button is clicked toggle the state of the menu anchor which will open or close the menu
-	if (MenuAnchor->ShouldOpenDueToClick())
-	{
-		MenuAnchor->SetIsOpen(true);
-		SetOpenMenu(MenuAnchor);
-	}
-	else
-	{
-		MenuAnchor->SetIsOpen(false);
-		TSharedPtr<SMenuAnchor> NullAnchor;
-		SetOpenMenu(NullAnchor);
-	}
-
-	return FReply::Handled();
-}
-
 //***********************************************************************************
 //SPIEPreviewWindow Implementation
 //***********************************************************************************
 SPIEPreviewWindow::SPIEPreviewWindow()
 {
-
 }
 
-void SPIEPreviewWindow::Construct(const FArguments& InArgs, TSharedPtr<FPIEPreviewDevice> InDevice)
+SPIEPreviewWindow::~SPIEPreviewWindow()
 {
-	SWindow::Construct(InArgs);
-
-	SetDevice(InDevice);
+	PrepareShutdown();
 }
 
 TSharedRef<SWidget> SPIEPreviewWindow::MakeWindowTitleBar(const TSharedRef<SWindow>& Window, const TSharedPtr<SWidget>& CenterContent, EHorizontalAlignment CenterContentAlignment)
 {
-	TSharedRef<SPIEPreviewWindowTitleBar> WindowTitleBar = SNew(SPIEPreviewWindowTitleBar, Window, CenterContent, CenterContentAlignment)
+	/*TSharedRef<SPIEPreviewWindowTitleBar> */WindowTitleBar = SNew(SPIEPreviewWindowTitleBar, Window, CenterContent, EHorizontalAlignment::HAlign_Center)
 		.Visibility(EVisibility::SelfHitTestInvisible);
-	return WindowTitleBar;
+
+	return WindowTitleBar.ToSharedRef();
 }
 
 EHorizontalAlignment SPIEPreviewWindow::GetTitleAlignment()
@@ -133,20 +44,16 @@ EHorizontalAlignment SPIEPreviewWindow::GetTitleAlignment()
 	return EHorizontalAlignment::HAlign_Left;
 }
 
-void SPIEPreviewWindow::SetSceneViewportPadding(const FMargin& Margin)
-{
-	ContentSlot->SlotPadding = Margin;
-}
-
 void SPIEPreviewWindow::ComputeBezelOrientation()
 {
 	if (BezelImage.IsValid())
 	{
-		FVector2D WindowSize = GetClientSizeInScreen();
+ 		float Width = Device->GetWindowClientWidth();
+ 		float Height = Device->GetWindowClientHeight();
 
 		bool bBezelRotated = Device->IsDeviceFlipped();
 
-		float ScaleX = bBezelRotated ? WindowSize.X / WindowSize.Y : 1.0f;
+		float ScaleX = bBezelRotated ? Width / Height : 1.0f;
 		float ScaleY = bBezelRotated ? Inverse(ScaleX) : 1.0f;
 
 		FScale2D Scale = FScale2D(ScaleX, ScaleY);
@@ -157,18 +64,6 @@ void SPIEPreviewWindow::ComputeBezelOrientation()
 	}
 }
 
-void SPIEPreviewWindow::CreateMenuToolBar()
-{
-	AddOverlaySlot()
-		.Padding(5.0f, 3.0f + SWindowDefs::DefaultTitleBarSize, 0, 0)
-		.HAlign(HAlign_Left)
-		.VAlign(VAlign_Top)
-		[
-			SNew(SPIEToolbar)
-			.OnGetMenuContent(this, &SPIEPreviewWindow::BuildSettingsMenu)
-		];
-}
-
 TSharedRef<SWidget> SPIEPreviewWindow::BuildSettingsMenu()
 {
 	FMenuBuilder MenuBuilder(true, nullptr);
@@ -190,14 +85,14 @@ TSharedRef<SWidget> SPIEPreviewWindow::BuildSettingsMenu()
 		auto IsCheckedFunction = [this, ScaleFactor]()
 		{
 			float WindowScaleFactor = GetWindowScaleFactor();
+			ECheckBoxState CheckState = FMath::IsNearlyEqual(ScaleFactor, WindowScaleFactor) ? ECheckBoxState::Checked : ECheckBoxState::Unchecked;
 
-			return FMath::IsNearlyEqual(ScaleFactor, WindowScaleFactor) ? ECheckBoxState::Checked : ECheckBoxState::Unchecked;
+			return CheckState;
 		};
 
 		auto ExecuteActionFunction = [this, ScaleFactor]()
 		{
-			SetScaleWindowToDeviceSize(false);
-			ScaleWindow(ScaleFactor);
+			SetWindowScaleFactor(ScaleFactor);
 		};
 
 		CreateMenuEntry(MenuBuilder, MoveTemp(EntryText), MoveTemp(IsCheckedFunction), MoveTemp(ExecuteActionFunction));
@@ -209,13 +104,17 @@ TSharedRef<SWidget> SPIEPreviewWindow::BuildSettingsMenu()
 		FText EntryText = FText::FromString(TEXT("Scale to device size"));
 		auto IsCheckedFunction = [this]()
 		{
-			bool bScaleToDeviceSize = IsManualManageDPIChanges();
-			return bScaleToDeviceSize ? ECheckBoxState::Checked : ECheckBoxState::Unchecked;
+			float WindowScaleFactor = GetWindowScaleFactor();
+			float DeviceSizeFactor = GetScaleToDeviceSizeFactor();
+
+			ECheckBoxState CheckState = FMath::IsNearlyEqual(WindowScaleFactor, DeviceSizeFactor) ? ECheckBoxState::Checked : ECheckBoxState::Unchecked;
+			return CheckState;
 		};
 
 		auto ExecuteActionFunction = [this]()
 		{
-			SetScaleWindowToDeviceSize(true);
+			float WindowScaleFactor = GetScaleToDeviceSizeFactor();
+			SetWindowScaleFactor(WindowScaleFactor);
 		};
 
 		CreateMenuEntry(MenuBuilder, MoveTemp(EntryText), MoveTemp(IsCheckedFunction), MoveTemp(ExecuteActionFunction));
@@ -257,9 +156,105 @@ TSharedRef<SWidget> SPIEPreviewWindow::BuildSettingsMenu()
 		CreateMenuEntry(MenuBuilder, MoveTemp(EntryText), MoveTemp(IsCheckedFunction), MoveTemp(ExecuteActionFunction));
 	}
 
+	MenuBuilder.AddMenuSeparator();
+	auto ResolutionDescriptionWidget = SNew(STextBlock)
+		.Text(FText::FromString(TEXT("Resolution")))
+		.Justification(ETextJustify::Center);
+
+	MenuBuilder.AddWidget(ResolutionDescriptionWidget, FText());
+
+	// Base resolution text
+	{
+		auto PrintLambda = [Device = Device]()
+		{
+			FString ResolutionText = TEXT("Device - ");
+			if (Device.IsValid())
+			{
+				int32 ResX, ResY;
+				Device->GetDeviceDefaultResolution(ResX, ResY);
+				if (Device->IsDeviceFlipped())
+				{
+					Swap(ResX, ResY);
+				}
+
+				ResolutionText += FString::FromInt(ResX) + TEXT("x") + FString::FromInt(ResY);
+			}
+
+			return FText::FromString(ResolutionText);
+		};
+
+		CreateTextMenuEntry(MenuBuilder, MoveTemp(PrintLambda));
+	}
+	// End Base resolution text
+
+	// Resolution with content scale
+	{
+		auto PrintLambda = [Device = Device]()
+		{
+			FString ResolutionText = TEXT("Content - ");
+			if (Device.IsValid())
+			{
+				int32 ResX, ResY;
+				bool bDeviceIgnoresContentFactor = Device->GetIgnoreMobileContentScaleFactor();
+				Device->SetIgnoreMobileContentScaleFactor(false);
+				Device->ComputeContentScaledResolution(ResX, ResY);
+				Device->SetIgnoreMobileContentScaleFactor(bDeviceIgnoresContentFactor);
+				if (Device->IsDeviceFlipped())
+				{
+					Swap(ResX, ResY);
+				}
+				ResolutionText += FString::FromInt(ResX) + TEXT("x") + FString::FromInt(ResY);
+			}
+
+			return FText::FromString(ResolutionText);
+		};
+
+		CreateTextMenuEntry(MenuBuilder, MoveTemp(PrintLambda));
+	}
+	// Resolution with content scale
+
+	// Displayed resolution
+	{
+		auto PrintLambda = [Device = Device]()
+		{
+			FString ResolutionText = TEXT("Window - ");
+			if (Device.IsValid())
+			{
+				int32 ResX, ResY;
+				Device->ComputeDeviceResolution(ResX, ResY);
+				if (Device->IsDeviceFlipped())
+				{
+					Swap(ResX, ResY);
+				}
+				ResolutionText += FString::FromInt(ResX) + TEXT("x") + FString::FromInt(ResY);
+			}
+
+			return FText::FromString(ResolutionText);
+		};
+
+		CreateTextMenuEntry(MenuBuilder, MoveTemp(PrintLambda));
+	}
+	// Resolution with content scale
+
 	return MenuBuilder.MakeWidget();
 }
 
+void SPIEPreviewWindow::CreateTextMenuEntry(class FMenuBuilder& MenuBuilder, TFunction<FText()>&& CreateTextFunction)
+{
+	auto Box = SNew(SHorizontalBox)
+		+ SHorizontalBox::Slot()
+		.FillWidth(1.0f)
+		.HAlign(EHorizontalAlignment::HAlign_Fill)
+		[
+			SNew(STextBlock)
+			.Visibility(EVisibility::HitTestInvisible)
+			.Justification(ETextJustify::Center)
+			.Text_Lambda(MoveTemp(CreateTextFunction))
+		];
+
+	MenuBuilder.AddWidget(Box, FText());
+}
+
 void SPIEPreviewWindow::CreateMenuEntry(FMenuBuilder& MenuBuilder, FText&& TextEntry, TFunction<ECheckBoxState()>&& IsCheckedFunction, TFunction<void()>&& ExecuteActionFunction)
 {
 	auto Box =
@@ -267,6 +262,7 @@ void SPIEPreviewWindow::CreateMenuEntry(FMenuBuilder& MenuBuilder, FText&& TextE
 		+ SHorizontalBox::Slot()
 		.AutoWidth()
 		.HAlign(EHorizontalAlignment::HAlign_Left)
+		.Padding(FMargin(5.0f, 0.0f))
 		[
 			SNew(STextBlock)
 			.Visibility(EVisibility::HitTestInvisible)
@@ -322,110 +318,134 @@ void SPIEPreviewWindow::CreatePIEPreviewBezelOverlay(UTexture2D* pBezelImage)
 	}
 }
 
-void SPIEPreviewWindow::ScaleWindow(const float ScreenFactor, const float DPIScaleFactor /*= 1.0f*/)
+void SPIEPreviewWindow::ValidatePosition(FVector2D& WindowPos)
 {
-	if (!Device.IsValid())
-	{
-		return;
-	}
+	WindowPos.X = FMath::CeilToInt(WindowPos.X);
+	WindowPos.Y = FMath::CeilToInt(WindowPos.Y);
+
+	FDisplayMetrics DisplayMetrics;
+	FDisplayMetrics::RebuildDisplayMetrics(DisplayMetrics);
 
-	if (FMath::IsNearlyEqual(ScreenFactor, CachedScaleToDeviceFactor) && FMath::IsNearlyEqual(DPIScaleFactor, CachedDPIScaleFactor))
+	const int32 Offset = 5;
+	if ((WindowPos.X - Offset > DisplayMetrics.VirtualDisplayRect.Right) ||
+		(WindowPos.X + Device->GetWindowWidth() + Offset < DisplayMetrics.VirtualDisplayRect.Left) ||
+		(WindowPos.Y - Offset > DisplayMetrics.VirtualDisplayRect.Bottom) ||
+		(WindowPos.Y + Device->GetWindowHeight() + Offset < DisplayMetrics.VirtualDisplayRect.Top))
 	{
-		return;
+		WindowPos.X = (DisplayMetrics.PrimaryDisplayWorkAreaRect.Left + DisplayMetrics.PrimaryDisplayWorkAreaRect.Right) / 2;
+		WindowPos.Y = (DisplayMetrics.PrimaryDisplayWorkAreaRect.Bottom + DisplayMetrics.PrimaryDisplayWorkAreaRect.Top) / 2;
 	}
+}
 
-	CachedScaleToDeviceFactor = ScreenFactor;
-	CachedDPIScaleFactor = DPIScaleFactor;
+void SPIEPreviewWindow::PrepareWindow(FVector2D WindowPosition, const float InitialScaleFactor, TSharedPtr<FPIEPreviewDevice> PreviewDevice)
+{
+	// we always manually handle DPI changes
+	SetManualManageDPIChanges(true);
 
-	if (IsManualManageDPIChanges())
-	{
-		FSlateApplication::Get().HandleDPIScaleChanged(GetNativeWindow().ToSharedRef());
-	}
+	SetDevice(PreviewDevice);
+
+	// place window to the required position and compute its size
+	ValidatePosition(WindowPosition);
+	MoveWindowTo(WindowPosition);
+	SetWindowScaleFactor(InitialScaleFactor, false);
 
-	Device->ScaleResolution(ScreenFactor, DPIScaleFactor, bClampWindowSizeState);
+	// update display resolution
+	const int32 ClientWidth = Device->GetWindowWidth();
+	const int32 ClientHeight = Device->GetWindowHeight() - GetTitleBarSize().Get();
+	FSystemResolution::RequestResolutionChange(ClientWidth, ClientHeight, EWindowMode::Windowed);
+	IConsoleManager::Get().CallAllConsoleVariableSinks();
 
+	// the above call will reset the position of the window and set the wrong size (due to manual DPI) and we need to set it right
+	MoveWindowTo(WindowPosition);
 	UpdateWindow();
+
+	// set needed event callbacks
+	SetOnWindowMoved(FOnWindowMoved::CreateSP(this, &SPIEPreviewWindow::OnWindowMoved));
+	HandleDPIChange = FSlateApplication::Get().OnSystemSignalsDPIChanged().AddSP(this, &SPIEPreviewWindow::OnDisplayDPIChanged);
 }
 
-float SPIEPreviewWindow::GetWindowScaleFactor() const
+void SPIEPreviewWindow::SetWindowScaleFactor(const float ScaleFactor, const bool bStore/* = true*/)
 {
-	float ScaleFactor = 1.0f;
+	WindowScalingFactor = ScaleFactor;
 
-	if (Device.IsValid())
+	// when required we will save the scaling value so it can be restored after session restart
+	if (bStore)
 	{
-		ScaleFactor = Device->GetResolutionScale();
+		GConfig->SetFloat(TEXT("/Script/Engine.MobilePIE"), TEXT("DeviceScalingFactor"), WindowScalingFactor, GEngineIni);
 	}
 
-	return ScaleFactor;
+	ScaleWindow(ScaleFactor);
 }
 
-void SPIEPreviewWindow::OnWindowMoved(const TSharedRef<SWindow>& Window)
+void SPIEPreviewWindow::ScaleWindow(float ScaleFactor)
 {
-	if (IsManualManageDPIChanges())
+	if (!Device.IsValid())
 	{
-		ScaleDeviceToPhisicalSize();
+		return;
 	}
 
-	FVector2D WindowPos = GetPositionInScreen();
-	GConfig->SetInt(TEXT("/Script/Engine.MobilePIE"), TEXT("WindowPosX"), FMath::CeilToInt(WindowPos.X), GEngineIni);
-	GConfig->SetInt(TEXT("/Script/Engine.MobilePIE"), TEXT("WindowPosY"), FMath::CeilToInt(WindowPos.Y), GEngineIni);
-}
+	bool bScaleToDeviceSize = IsScalingToDeviceSizeFactor(ScaleFactor);
+	Device->SetIgnoreMobileContentScaleFactor(bScaleToDeviceSize);
 
-void SPIEPreviewWindow::OnDisplayDPIChanged(TSharedRef<SWindow> Window)
-{
-	if (IsManualManageDPIChanges())
+	float DPIScaleFactor = ComputeDPIScaleFactor();
+
+	if (bScaleToDeviceSize)
 	{
-		ScaleDeviceToPhisicalSize();
+		ScaleFactor = ComputeScaleToDeviceSizeFactor();
+		ScaleFactor /= DPIScaleFactor;
 	}
-}
 
-void SPIEPreviewWindow::SetScaleWindowToDeviceSize(const bool bScale)
-{
-	if (bScale != IsManualManageDPIChanges())
+	if (FMath::IsNearlyEqual(ScaleFactor, CachedScaleToDeviceFactor) && FMath::IsNearlyEqual(DPIScaleFactor, CachedDPIScaleFactor))
 	{
-		SetManualManageDPIChanges(bScale);
-
-		if (Device.IsValid())
-		{
-			Device->SetIgnoreMobileContentScaleFactor(bScale);
-		}
+		return;
+	}
 
-		if (IsManualManageDPIChanges())
-		{
-			ScaleDeviceToPhisicalSize();
+	CachedScaleToDeviceFactor = ScaleFactor;
+	CachedDPIScaleFactor = DPIScaleFactor;
 
-			SetOnWindowMoved(FOnWindowMoved::CreateSP(this, &SPIEPreviewWindow::OnWindowMoved));
-			HandleDPIChange = FSlateApplication::Get().OnSystemSignalsDPIChanged().AddSP(this, &SPIEPreviewWindow::OnDisplayDPIChanged);
-		}
-		else
-		{
-			SetOnWindowMoved(nullptr);
+	SetDPIScaleFactor(CachedDPIScaleFactor);
 
-			if (HandleDPIChange.IsValid())
-			{
-				FSlateApplication::Get().OnSystemSignalsDPIChanged().Remove(HandleDPIChange);
-			}
-		}
+	if (IsManualManageDPIChanges())
+	{
+		FSlateApplication::Get().HandleDPIScaleChanged(GetNativeWindow().ToSharedRef());
 	}
+	
+	Device->ScaleResolution(ScaleFactor, DPIScaleFactor, bClampWindowSizeState);
+
+	UpdateWindow();
 }
 
-void SPIEPreviewWindow::ScaleDeviceToPhisicalSize()
+void SPIEPreviewWindow::OnWindowMoved(const TSharedRef<SWindow>& Window)
 {
-	if (Device.IsValid())
-	{
-		float ScreenFactor;
-		float DPIScaleFactor;
-		ComputeScaleToDeviceSizeFactor(ScreenFactor, DPIScaleFactor);
+	float CurrentScaleFactor = GetWindowScaleFactor();
+	ScaleWindow(CurrentScaleFactor);
 
-		SetDPIScaleFactor(DPIScaleFactor);
-		ScaleWindow(ScreenFactor, DPIScaleFactor);
-	}
+	// save the position so we can restore it if the session is restarted
+	FVector2D WindowPos = GetPositionInScreen();
+	GConfig->SetInt(TEXT("/Script/Engine.MobilePIE"), TEXT("WindowPosX"), FMath::CeilToInt(WindowPos.X), GEngineIni);
+	GConfig->SetInt(TEXT("/Script/Engine.MobilePIE"), TEXT("WindowPosY"), FMath::CeilToInt(WindowPos.Y), GEngineIni);
+}
+
+void SPIEPreviewWindow::OnDisplayDPIChanged(TSharedRef<SWindow> Window)
+{
+	float CurrentScaleFactor = GetWindowScaleFactor();
+	SetWindowScaleFactor(CurrentScaleFactor);
 }
 
-void SPIEPreviewWindow::ComputeScaleToDeviceSizeFactor(float& OutScreenFactor, float& OutDPIScaleFactor) const
+float SPIEPreviewWindow::ComputeDPIScaleFactor()
 {
-	OutScreenFactor = 1.0f;
-	OutDPIScaleFactor = 1.0f;
+	FVector2D WindowPos = GetPositionInScreen();
+	int32 PointX = FMath::RoundToInt(WindowPos.X);
+	int32 PointY = FMath::RoundToInt(WindowPos.Y);
+
+	float DPIFactor = FPlatformApplicationMisc::GetDPIScaleFactorAtPoint(PointX, PointY);
+
+	return DPIFactor;
+}
+
+float SPIEPreviewWindow::ComputeScaleToDeviceSizeFactor() const
+{
+	float OutScreenFactor = 1.0f;
 
 	FDisplayMetrics DisplayMetrics;
 	FDisplayMetrics::RebuildDisplayMetrics(DisplayMetrics);
@@ -469,11 +489,11 @@ void SPIEPreviewWindow::ComputeScaleToDeviceSizeFactor(float& OutScreenFactor, f
 		}
 	}
 
-	OutDPIScaleFactor = FPlatformApplicationMisc::GetDPIScaleFactorAtPoint(PointX, PointY);
-
 	const int32 DevicePPI = Device->GetDeviceSpecs()->PPI;
 	float PPIRatio = (!!DevicePPI && !!LocalPPI) ? (float)LocalPPI / (float)DevicePPI : 1.0f;
-	OutScreenFactor = PPIRatio * RatioMonitorResolution / OutDPIScaleFactor;
+	OutScreenFactor = PPIRatio * RatioMonitorResolution;
+
+	return OutScreenFactor;
 }
 
 void SPIEPreviewWindow::RotateWindow()
@@ -485,6 +505,8 @@ void SPIEPreviewWindow::RotateWindow()
 
 	Device->SwitchOrientation(bClampWindowSizeState);
 
+	UpdateGameLayerManagerDefaultViewport();
+
 	UpdateWindow();
 }
 
@@ -544,7 +566,7 @@ void SPIEPreviewWindow::UpdateWindow()
 	ReshapeWindow(FVector2D(PosX, PosY), FVector2D(Device->GetWindowWidth(), Device->GetWindowHeight()));
 
 	// offset the viewport widget into its correct location
-	SetSceneViewportPadding(Device->GetViewportMargin());
+	ContentSlot->SlotPadding = Device->GetViewportMargin();
 
 	// bezel orientation depends on the window size so we need to call it after ReshapeWindow()
 	ComputeBezelOrientation();
@@ -556,13 +578,31 @@ void SPIEPreviewWindow::SetDevice(TSharedPtr<FPIEPreviewDevice> InDevice)
 
 	if (Device.IsValid())
 	{
-		FMargin ViewportOffset = Device->GetViewportMargin();
-		SetSceneViewportPadding(ViewportOffset);
-
 		CreatePIEPreviewBezelOverlay(InDevice->GetBezelTexture());
+	}
+}
+
+void SPIEPreviewWindow::PrepareShutdown()
+{
+	SetOnWindowMoved(nullptr);
+
+	if (HandleDPIChange.IsValid())
+	{
+		if (FSlateApplication::IsInitialized())
+		{
+			FSlateApplication::Get().OnSystemSignalsDPIChanged().Remove(HandleDPIChange);
+		}
+	}
 
-		CreateMenuToolBar();
+	if (BezelImage.IsValid())
+	{
+		RemoveOverlaySlot(BezelImage.ToSharedRef());
+
+		BezelBrush.SetResourceObject(nullptr);
+		BezelImage = nullptr;
 	}
+
+	Device = nullptr;
 }
 
 int32 SPIEPreviewWindow::GetDefaultTitleBarSize()
@@ -570,4 +610,27 @@ int32 SPIEPreviewWindow::GetDefaultTitleBarSize()
 	return SWindowDefs::DefaultTitleBarSize;
 }
 
+void SPIEPreviewWindow::SetGameLayerManagerWidget(TSharedPtr<class SGameLayerManager> GameLayerManager)
+{
+	GameLayerManagerWidget = GameLayerManager;
+
+	UpdateGameLayerManagerDefaultViewport();
+}
+
+void SPIEPreviewWindow::UpdateGameLayerManagerDefaultViewport()
+{
+	if (Device.IsValid() && GameLayerManagerWidget.IsValid())
+	{
+		FIntPoint DeviceResolution;
+		Device->GetDeviceDefaultResolution(DeviceResolution.X, DeviceResolution.Y);
+
+		if (Device->IsDeviceFlipped())
+		{
+			Swap(DeviceResolution.X, DeviceResolution.Y);
+		}
+
+		GameLayerManagerWidget->SetUseFixedDPIValue(true, DeviceResolution);
+	}
+}
+
 #endif
\ No newline at end of file
diff --git a/Engine/Source/Editor/PIEPreviewDeviceProfileSelector/Private/PIEPreviewWindow.h b/Engine/Source/Editor/PIEPreviewDeviceProfileSelector/Private/PIEPreviewWindow.h
--- a/Engine/Source/Editor/PIEPreviewDeviceProfileSelector/Private/PIEPreviewWindow.h
+++ b/Engine/Source/Editor/PIEPreviewDeviceProfileSelector/Private/PIEPreviewWindow.h
@@ -17,23 +17,23 @@ public:
 	*/
 	SPIEPreviewWindow();
 
-	void Construct(const FArguments& InArgs, TSharedPtr<FPIEPreviewDevice> InDevice);
+	~SPIEPreviewWindow();
 
-	// call this to 
-	void SetSceneViewportPadding(const FMargin& Margin);
+	// call this function set the 'preview device' and set initial values for the position and resolution scaling factor
+	void PrepareWindow(FVector2D WindowPosition, const float InitialScaleFactor, TSharedPtr<FPIEPreviewDevice> PreviewDevice);
 
-	// sets the device displayed by this window
-	void SetDevice(TSharedPtr<class FPIEPreviewDevice> InDevice);
+	// call this function to perform cleanup
+	void PrepareShutdown();
 
 	// retrieves the display device
 	TSharedPtr<class FPIEPreviewDevice> GetDevice();
 
-	// call to perform resolution scaling on the device
-	void ScaleWindow(const float ScaleFactor, const float BezelFactor = 1.0f);
-
-	// call to retrieve device's current resolution scale factor
+	// call to retrieve the window current resolution scale factor
 	float GetWindowScaleFactor() const;
 
+	// call to set the window resolution scale factor
+	void SetWindowScaleFactor(const float ScaleFactor, const bool bStore = true);
+
 	// call to rotate the window emulating a rotation of the physical device
 	void RotateWindow();
 
@@ -48,42 +48,58 @@ public:
 
 	static int32 GetDefaultTitleBarSize();
 
-	// call with a true argument to resize the window so it matches the device physical size
-	void SetScaleWindowToDeviceSize(const bool bScale);
-
 	// call with a true argument to restrict the window size to the desktop size
 	void SetClampWindowSize(const bool bClamp);
 	bool IsClampingWindowSize() const;
 
+	/** creates and returns the settings menu */
+	TSharedRef<SWidget> BuildSettingsMenu();
+
+	/** we need the game layer manager to control the DPI scaling behavior and this function can be called should be called when the manager is available */
+	void SetGameLayerManagerWidget(TSharedPtr<class SGameLayerManager> GameLayerManager);
+
 private:
 	virtual EHorizontalAlignment GetTitleAlignment() override;
 	virtual TSharedRef<SWidget> MakeWindowTitleBar(const TSharedRef<SWindow>& Window, const TSharedPtr<SWidget>& CenterContent, EHorizontalAlignment CenterContentAlignment) override;
 
+	// sets the device displayed by this window
+	void SetDevice(TSharedPtr<class FPIEPreviewDevice> InDevice);
+
+	// returns whether or not the provided ScaleFactor is a 'scale to physical device size' factor
+	bool IsScalingToDeviceSizeFactor(const float ScaleFactor) const;
+
+	// returns the scaling value reserved for 'scale to physical device size' factor
+	float GetScaleToDeviceSizeFactor() const;
+
+	// utility function used to correct a given window position if it's out of the display area
+	void ValidatePosition(FVector2D& WindowPos);
+
+	// call to perform resolution scaling on the device
+	void ScaleWindow(float ScaleFactor);
+
 	// this will properly scale and rotate the bezel to match the orientation of the device
 	void ComputeBezelOrientation();
 
 	// function that adds a widget to display the device's bezel
 	void CreatePIEPreviewBezelOverlay(UTexture2D* pBezelImage);
 
-	// function used create a menu toolbar
-	void CreateMenuToolBar();
-
-	/** creates and returns the settings menu */
-	TSharedRef<SWidget> BuildSettingsMenu();
-
 	/** helper function used to create a menu item formed from a description text and a checkbox */
 	/** @ IsCheckedFunction - computes the state of the checkbox */
 	/** @ ExecuteActionFunction - function to be called when this menu item is clicked */
 	void CreateMenuEntry(class FMenuBuilder& MenuBuilder, FText&& TextEntry, TFunction<ECheckBoxState()>&& IsCheckedFunction, TFunction<void()>&& ExecuteActionFunction);
 
+	/** helper function used to create a simple text widget to the provided menu */
+	/** @ CreateTextFunction - function should return the menu text that we want to display */
+	void CreateTextMenuEntry(class FMenuBuilder& MenuBuilder, TFunction<FText()>&& CreateTextFunction);
+
 	// call to compute window size, position on screen and bezel orientation
 	void UpdateWindow();
 
-	// call function to make the size of this window match the physical device size
-	void ScaleDeviceToPhisicalSize();
-
 	// utility function that will compute a screen and DPI scale factor needed to scale the display window to the physical device size
-	void ComputeScaleToDeviceSizeFactor(float& OutScreenFactor, float& OutDPIScaleFactor) const;
+	float ComputeScaleToDeviceSizeFactor() const;
+
+	// utility function used to retrieve the DPI scaling factor based on window's position
+	float ComputeDPIScaleFactor();
 
 	// callback used when DPI value varies, useful when we need to constrain the window to the physical device size
 	void OnWindowMoved(const TSharedRef<SWindow>& Window);
@@ -91,6 +107,10 @@ private:
 	// callback used when DPI value varies, useful when we need to constrain the window to the physical device size
 	void OnDisplayDPIChanged(TSharedRef<SWindow> Window);
 
+	// utility function used to control game layer manager's DPI scaling behavior
+	// this should be called every time the window is rotated
+	void UpdateGameLayerManagerDefaultViewport();
+
 protected:
 	// brush created to display the bezel
 	FSlateBrush BezelBrush;
@@ -107,7 +127,14 @@ protected:
 	/** when true the window size will be restricted to the desktop size  */
 	bool bClampWindowSizeState = true;
 
+	float WindowScalingFactor = 0.0f;
+
+	TSharedPtr<SPIEPreviewWindowTitleBar> WindowTitleBar;
+
 	FDelegateHandle HandleDPIChange;
+
+	/** Pointer to the game layer manager widget. This is needed because we want to control the DPI scaling behavior */
+	TSharedPtr<SGameLayerManager> GameLayerManagerWidget;
 };
 
 FORCEINLINE bool SPIEPreviewWindow::IsRotationAllowed() const
@@ -135,4 +162,20 @@ FORCEINLINE bool SPIEPreviewWindow::IsClampingWindowSize() const
 	return bClampWindowSizeState;
 }
 
+FORCEINLINE bool SPIEPreviewWindow::IsScalingToDeviceSizeFactor(const float ScaleFactor) const
+{
+	float DeviceSizeFactor = GetScaleToDeviceSizeFactor();
+	return FMath::IsNearlyEqual(ScaleFactor, DeviceSizeFactor);
+}
+
+FORCEINLINE float SPIEPreviewWindow::GetWindowScaleFactor() const
+{
+	return WindowScalingFactor;
+}
+
+FORCEINLINE float SPIEPreviewWindow::GetScaleToDeviceSizeFactor() const
+{
+	return 0.0f;
+}
+
 #endif
\ No newline at end of file
diff --git a/Engine/Source/Editor/PIEPreviewDeviceProfileSelector/Private/PIEPreviewWindowTitleBar.cpp b/Engine/Source/Editor/PIEPreviewDeviceProfileSelector/Private/PIEPreviewWindowTitleBar.cpp
--- a/Engine/Source/Editor/PIEPreviewDeviceProfileSelector/Private/PIEPreviewWindowTitleBar.cpp
+++ b/Engine/Source/Editor/PIEPreviewDeviceProfileSelector/Private/PIEPreviewWindowTitleBar.cpp
@@ -4,10 +4,93 @@
 #include "PIEPreviewWindow.h"
 #include "Widgets/Input/SCheckBox.h"
 #include "Framework/MultiBox/MultiBoxBuilder.h"
+#include "SViewportToolBar.h"
 
 #if WITH_EDITOR
 
-#define LOCTEXT_NAMESPACE "PIEPreviewWindowTitleBar"
+/** Toolbar class used to add some menus to configure various device display settings */
+class SPIEToolbar : public SViewportToolBar
+{
+public:
+	SLATE_BEGIN_ARGS(SPIEToolbar) {}
+		SLATE_EVENT(FOnGetContent, OnGetMenuContent)
+	SLATE_END_ARGS()
+
+	void Construct(const FArguments& InArgs);
+	FReply OnMenuClicked();
+
+protected:
+	TSharedPtr<SMenuAnchor> MenuAnchor;
+};
+
+void SPIEToolbar::Construct(const FArguments& InArgs)
+{
+	SViewportToolBar::Construct(SViewportToolBar::FArguments());
+
+	const FSlateBrush* ImageBrush = FPIEPreviewWindowCoreStyle::Get().GetBrush("ComboButton.Arrow");
+
+	TSharedPtr<SWidget> ButtonContent =
+		SNew(SHorizontalBox)
+		+ SHorizontalBox::Slot()
+		.AutoWidth()
+		.Padding(2.0f, 2.0f)
+		[
+			SNew(SHorizontalBox)
+			+ SHorizontalBox::Slot()
+			.AutoWidth()
+			[
+				SNew(SBox)
+				.VAlign(VAlign_Center)
+				.HAlign(HAlign_Center)
+				.IsEnabled(true)
+				.Cursor(EMouseCursor::Default)
+				[
+					SNew(SImage)
+					.Image(ImageBrush)
+					.ColorAndOpacity(FSlateColor::UseForeground())
+				]
+			]
+		];
+
+	ChildSlot
+	[
+		SAssignNew(MenuAnchor, SMenuAnchor)
+		.Padding(0)
+		.Placement(MenuPlacement_BelowAnchor)
+		[
+			SNew(SButton)
+			// Allows users to drag with the mouse to select options after opening the menu */
+			.ClickMethod(EButtonClickMethod::MouseDown)
+			.ContentPadding(FMargin(2.0f, 2.0f))
+			.VAlign(VAlign_Center)
+			.ButtonStyle(FPIEPreviewWindowCoreStyle::Get(), "PIEWindow.MenuButton")
+			.OnClicked(this, &SPIEToolbar::OnMenuClicked)
+			[
+				ButtonContent.ToSharedRef()
+			]
+		]
+		.OnGetMenuContent(InArgs._OnGetMenuContent)
+	];
+}
+
+FReply SPIEToolbar::OnMenuClicked()
+{
+	// If the menu button is clicked toggle the state of the menu anchor which will open or close the menu
+	if (MenuAnchor->ShouldOpenDueToClick())
+	{
+		MenuAnchor->SetIsOpen(true);
+		SetOpenMenu(MenuAnchor);
+	}
+	else
+	{
+		MenuAnchor->SetIsOpen(false);
+		TSharedPtr<SMenuAnchor> NullAnchor;
+		SetOpenMenu(NullAnchor);
+	}
+
+	return FReply::Handled();
+}
+
 
 void SPIEPreviewWindowTitleBar::MakeTitleBarContentWidgets(TSharedPtr< SWidget >& OutLeftContent, TSharedPtr< SWidget >& OutRightContent)
 {
@@ -23,11 +106,8 @@ void SPIEPreviewWindowTitleBar::MakeTitleBarContentWidgets(TSharedPtr< SWidget >
 			[this]()
 			{
 				TSharedPtr<SPIEPreviewWindow> PIEWindow = GetOwnerWindow();
-
-				if (PIEWindow.IsValid())
-				{
-					PIEWindow->RotateWindow();
-				}
+				check(PIEWindow.IsValid());
+				PIEWindow->RotateWindow();
 
 				return FReply::Handled();
 			}
@@ -44,13 +124,9 @@ void SPIEPreviewWindowTitleBar::MakeTitleBarContentWidgets(TSharedPtr< SWidget >
 			[this]()
 			{
 				TSharedPtr<SPIEPreviewWindow> PIEWindow = GetOwnerWindow();
+				check(PIEWindow.IsValid());
 			
-				if (PIEWindow.IsValid())
-				{
-					return PIEWindow->IsRotationAllowed();
-				}
-
-				return false;
+				return PIEWindow->IsRotationAllowed();
 			}
 		)
 	;
@@ -65,6 +141,23 @@ void SPIEPreviewWindowTitleBar::MakeTitleBarContentWidgets(TSharedPtr< SWidget >
 			ScreenRotationButton.ToSharedRef()
 		];
 
+	// Add the settings menu widget
+	WindowTitleBarButtons->AddSlot()
+		.AutoWidth()
+		[
+			SNew(SPIEToolbar)
+			.OnGetMenuContent_Lambda
+			(
+				[this]()
+				{
+					TSharedPtr<SPIEPreviewWindow> PIEWindow = GetOwnerWindow();
+					check(PIEWindow.IsValid());
+
+					return PIEWindow->BuildSettingsMenu();
+				}
+			)
+		];
+
 	if (OutRightContentBaseWindow.IsValid())
 	{
 		WindowTitleBarButtons->AddSlot()
@@ -105,6 +198,4 @@ const FSlateBrush* SPIEPreviewWindowTitleBar::GetScreenRotationButtonImage() con
 	}
 }
 
-#undef LOCTEXT_NAMESPACE
-
 #endif
diff --git a/Engine/Source/Editor/PIEPreviewDeviceProfileSelector/Public/PIEPreviewDeviceProfileSelectorModule.h b/Engine/Source/Editor/PIEPreviewDeviceProfileSelector/Public/PIEPreviewDeviceProfileSelectorModule.h
--- a/Engine/Source/Editor/PIEPreviewDeviceProfileSelector/Public/PIEPreviewDeviceProfileSelectorModule.h
+++ b/Engine/Source/Editor/PIEPreviewDeviceProfileSelector/Public/PIEPreviewDeviceProfileSelectorModule.h
@@ -42,7 +42,7 @@ public:
 	virtual TSharedRef<SWindow> CreatePIEPreviewDeviceWindow(FVector2D ClientSize, FText WindowTitle, EAutoCenter AutoCenterType, FVector2D ScreenPosition, TOptional<float> MaxWindowWidth, TOptional<float> MaxWindowHeight) override;
 
 	/** call this after the window is created and registered to the application to setup display related parameters */
-	virtual void PrepareDeviceDisplay() override;
+	virtual void OnWindowReady(TSharedRef<SWindow> Window) override;
 	
 	virtual const FPIEPreviewDeviceContainer& GetPreviewDeviceContainer() ;
 	TSharedPtr<FPIEPreviewDeviceContainerCategory> GetPreviewDeviceRootCategory() const { return EnumeratedDevices.GetRootCategory(); }
@@ -53,12 +53,21 @@ public:
 		return FParse::Value(FCommandLine::Get(), GetPreviewDeviceCommandSwitch(), PreviewDeviceDummy);
 	}
 
+	/** we need the game layer manager to control the DPI scaling behavior and this function can be called should be called when the manager is available */
+	virtual void SetGameLayerManagerWidget(TSharedPtr<class SGameLayerManager> GameLayerManager) override;
+
 private:
 	static const TCHAR* GetPreviewDeviceCommandSwitch()
 	{
 		return TEXT("MobileTargetDevice=");
 	}
 
+	/** callback function registered in UGameViewportClient::OnViewportCreated needed to disable mouse capture/lock	*/
+	void OnViewportCreated();
+
+	/** callback function registered in FCoreDelegates::OnFEngineLoopInitComplete needed to position and show the window */
+	void OnEngineInitComplete();
+
 	void InitPreviewDevice();
 	static FString GetDeviceSpecificationContentDir();
 	bool ReadDeviceSpecification();
@@ -66,14 +75,28 @@ private:
 
 	void UpdateDisplayResolution();
 
+	/** this function will attempt to load the last known window position and scaling factor */
+	bool ReadWindowConfig();
+
 private:
 	bool bInitialized;
 	FString DeviceProfile;
 	FString PreviewDevice;
 
+	/** delegate handle that will be obtained from UGameViewportClient::OnViewportCreated */ 
+	FDelegateHandle ViewportCreatedDelegate;
+
+	/** delegate handle that will be obtained from FCoreDelegates::OnFEngineLoopInitComplete */
+	FDelegateHandle EngineInitCompleteDelegate;
+
 	FPIEPreviewDeviceContainer EnumeratedDevices;
 
 	TSharedPtr<class FPIEPreviewDevice> Device;
 
 	TWeakPtr<class SPIEPreviewWindow> WindowWPtr;
+
+	FVector2D InitialWindowPosition;
+	float InitialWindowScaleValue;
+
+	TSharedPtr<class SGameLayerManager> GameLayerManagerWidget;
 };
diff --git a/Engine/Source/Runtime/Engine/Private/GameEngine.cpp b/Engine/Source/Runtime/Engine/Private/GameEngine.cpp
--- a/Engine/Source/Runtime/Engine/Private/GameEngine.cpp
+++ b/Engine/Source/Runtime/Engine/Private/GameEngine.cpp
@@ -185,6 +185,14 @@ void UGameEngine::CreateGameViewportWidget( UGameViewportClient* GameViewportCli
 {
 	bool bRenderDirectlyToWindow = (!StartupMovieCaptureHandle.IsValid() || IMovieSceneCaptureModule::Get().IsStereoAllowed()) && GIsDumpingMovie == 0;
 
+	TSharedRef<SOverlay> ViewportOverlayWidgetRef = SNew( SOverlay );
+
+	TSharedRef<SGameLayerManager> GameLayerManagerRef = SNew(SGameLayerManager)
+		.SceneViewport_UObject(this, &UGameEngine::GetGameSceneViewport, GameViewportClient)
+		[
+			ViewportOverlayWidgetRef
+		];
+
 	// when we're running in a "device simulation" window, render the scene to an intermediate texture
 	// in the mobile device "emulation" case this is needed to properly position the viewport (as a widget) inside its bezel
 #if WITH_EDITOR
@@ -192,17 +200,11 @@ void UGameEngine::CreateGameViewportWidget( UGameViewportClient* GameViewportCli
 	if (PIEPreviewDeviceModule && FPIEPreviewDeviceModule::IsRequestingPreviewDevice())
 	{
 		bRenderDirectlyToWindow = false;
+		PIEPreviewDeviceModule->SetGameLayerManagerWidget(GameLayerManagerRef);
 	}
 #endif
 
 	const bool bStereoAllowed = bRenderDirectlyToWindow;
-	TSharedRef<SOverlay> ViewportOverlayWidgetRef = SNew( SOverlay );
-
-	TSharedRef<SGameLayerManager> GameLayerManagerRef = SNew(SGameLayerManager)
-		.SceneViewport_UObject(this, &UGameEngine::GetGameSceneViewport, GameViewportClient)
-		[
-			ViewportOverlayWidgetRef
-		];
 
 	TSharedRef<SViewport> GameViewportWidgetRef = 
 		SNew( SViewport )
@@ -583,7 +585,7 @@ TSharedRef<SWindow> UGameEngine::CreateGameWindow()
 #if WITH_EDITOR
 	if (PIEPreviewDeviceModule && FPIEPreviewDeviceModule::IsRequestingPreviewDevice())
 	{
-		PIEPreviewDeviceModule->PrepareDeviceDisplay();
+		PIEPreviewDeviceModule->OnWindowReady(Window);
 	}
 #endif
 
diff --git a/Engine/Source/Runtime/Engine/Private/Slate/SGameLayerManager.cpp b/Engine/Source/Runtime/Engine/Private/Slate/SGameLayerManager.cpp
--- a/Engine/Source/Runtime/Engine/Private/Slate/SGameLayerManager.cpp
+++ b/Engine/Source/Runtime/Engine/Private/Slate/SGameLayerManager.cpp
@@ -23,6 +23,7 @@
 SGameLayerManager::SGameLayerManager()
 :	DefaultWindowTitleBarHeight(64.0f)
 ,	bIsGameUsingBorderlessWindow(false)
+,	bUseScaledDPI(false)
 {
 }
 
@@ -263,6 +264,17 @@ bool SGameLayerManager::OnVisualizeTooltip(const TSharedPtr<SWidget>& TooltipCon
 	return true;
 }
 
+void SGameLayerManager::SetUseFixedDPIValue(const bool bInUseFixedDPI, const FIntPoint ViewportSize /*= FIntPoint()*/)
+{
+	bUseScaledDPI = bInUseFixedDPI;
+	ScaledDPIViewportReference = ViewportSize;
+}
+
+bool SGameLayerManager::IsUsingFixedDPIValue() const
+{
+	return bUseScaledDPI;
+}
+
 float SGameLayerManager::GetGameViewportDPIScale() const
 {
 	const FSceneViewport* Viewport = SceneViewport.Get();
@@ -280,7 +292,19 @@ float SGameLayerManager::GetGameViewportDPIScale() const
 	}
 
 	FIntPoint ViewportSize = Viewport->GetSize();
-	const float GameUIScale = UserInterfaceSettings->GetDPIScaleBasedOnSize(ViewportSize);
+	float GameUIScale;
+
+	if (bUseScaledDPI)
+	{
+		float DPIValue = UserInterfaceSettings->GetDPIScaleBasedOnSize(ScaledDPIViewportReference);
+		float ViewportScale = FMath::Min((float)ViewportSize.X / (float)ScaledDPIViewportReference.X, (float)ViewportSize.Y / (float)ScaledDPIViewportReference.Y);
+
+		GameUIScale = DPIValue * ViewportScale;
+	}
+	else
+	{
+		GameUIScale = UserInterfaceSettings->GetDPIScaleBasedOnSize(ViewportSize);
+	}
 
 	// Remove the platform DPI scale from the incoming size.  Since the platform DPI is already
 	// attempt to normalize the UI for a high DPI, and the DPI scale curve is based on raw resolution
diff --git a/Engine/Source/Runtime/Engine/Public/IPIEPreviewDeviceModule.h b/Engine/Source/Runtime/Engine/Public/IPIEPreviewDeviceModule.h
--- a/Engine/Source/Runtime/Engine/Public/IPIEPreviewDeviceModule.h
+++ b/Engine/Source/Runtime/Engine/Public/IPIEPreviewDeviceModule.h
@@ -34,10 +34,15 @@ class IPIEPreviewDeviceModule : public IDeviceProfileSelectorModule
 		/**
 		 * should be called after the window is created and registered and before scene rendering begins
 		*/
-		virtual void PrepareDeviceDisplay() {};
+		virtual void OnWindowReady(TSharedRef<SWindow> Window) {};
 
 		/**
 		* Apply PieWindow device parameters
 		*/
 		virtual void ApplyPreviewDeviceState() = 0;
+
+		/** we need the game layer manager to control the DPI scaling behavior and this function can be called should be called when the manager is available */
+		virtual void SetGameLayerManagerWidget(TSharedPtr<class SGameLayerManager> GameLayerManager)
+		{
+		}
 };
\ No newline at end of file
diff --git a/Engine/Source/Runtime/Engine/Public/Slate/SGameLayerManager.h b/Engine/Source/Runtime/Engine/Public/Slate/SGameLayerManager.h
--- a/Engine/Source/Runtime/Engine/Public/Slate/SGameLayerManager.h
+++ b/Engine/Source/Runtime/Engine/Public/Slate/SGameLayerManager.h
@@ -136,6 +136,14 @@ public:
 	virtual bool OnVisualizeTooltip(const TSharedPtr<SWidget>& TooltipContent) override;
 	// End SWidget overrides
 
+	/**
+	 * Function will instruct internal DPI computations to use a provided reference viewport size instead of the actual viewport size.
+	 * After the DPI will be retrieved it will be scaled down with the ratio between the actual viewport size and the provided one.
+	 * Check GetGameViewportDPIScale() for more information.
+	 */
+	void SetUseFixedDPIValue(const bool bUseFixedDPI, const FIntPoint RefViewportSize = FIntPoint());
+	bool IsUsingFixedDPIValue() const;
+
 private:
 	float GetGameViewportDPIScale() const;
 	FOptionalSize GetDefaultWindowTitleBarHeight() const;
@@ -206,4 +214,7 @@ private:
 	TSharedPtr<SWidget> DefaultTitleBarContentWidget;
 	float DefaultWindowTitleBarHeight;
 	bool bIsGameUsingBorderlessWindow;
+
+	FIntPoint ScaledDPIViewportReference;
+	bool bUseScaledDPI;
 };
diff --git a/Engine/Source/Runtime/Launch/Private/LaunchEngineLoop.cpp b/Engine/Source/Runtime/Launch/Private/LaunchEngineLoop.cpp
--- a/Engine/Source/Runtime/Launch/Private/LaunchEngineLoop.cpp
+++ b/Engine/Source/Runtime/Launch/Private/LaunchEngineLoop.cpp
@@ -1838,8 +1838,13 @@ int32 FEngineLoop::PreInit(const TCHAR* CmdLine)
 
 	InitEngineTextLocalization();
 
+	bool bForceEnableHighDPI = false;
+#if WITH_EDITOR
+	bForceEnableHighDPI = FPIEPreviewDeviceModule::IsRequestingPreviewDevice();
+#endif
+
 	// This must be called before any window (including the splash screen is created
-	FSlateApplication::InitHighDPI();
+	FSlateApplication::InitHighDPI(bForceEnableHighDPI);
 
 	UStringTable::InitializeEngineBridge();
 
diff --git a/Engine/Source/Runtime/Slate/Private/Framework/Application/SlateApplication.cpp b/Engine/Source/Runtime/Slate/Private/Framework/Application/SlateApplication.cpp
--- a/Engine/Source/Runtime/Slate/Private/Framework/Application/SlateApplication.cpp
+++ b/Engine/Source/Runtime/Slate/Private/Framework/Application/SlateApplication.cpp
@@ -1110,7 +1110,7 @@ void FSlateApplication::SetupPhysicalSensitivities()
 	FGestureDetector::LongPressAllowedMovement = DragTriggerDistance;
 }
 
-void FSlateApplication::InitHighDPI()
+void FSlateApplication::InitHighDPI(const bool bForceEnable)
 {
 	IConsoleVariable* CVar = IConsoleManager::Get().FindConsoleVariable(TEXT("EnableHighDPIAwareness"));
 
@@ -1126,7 +1126,8 @@ void FSlateApplication::InitHighDPI()
 			GConfig->GetBool(TEXT("/Script/Engine.UserInterfaceSettings"), TEXT("bAllowHighDPIInGameMode"), bRequestEnableHighDPI, GEngineIni);
 		}
 
-		const bool bEnableHighDPI = bRequestEnableHighDPI && !FParse::Param(FCommandLine::Get(), TEXT("nohighdpi"));
+		bool bEnableHighDPI = bRequestEnableHighDPI && !FParse::Param(FCommandLine::Get(), TEXT("nohighdpi"));
+		bEnableHighDPI |= bForceEnable;
 
 		// Set the cvar here for other systems that need it.
 		CVar->Set(bEnableHighDPI);
diff --git a/Engine/Source/Runtime/Slate/Public/Framework/Application/SlateApplication.h b/Engine/Source/Runtime/Slate/Public/Framework/Application/SlateApplication.h
--- a/Engine/Source/Runtime/Slate/Public/Framework/Application/SlateApplication.h
+++ b/Engine/Source/Runtime/Slate/Public/Framework/Application/SlateApplication.h
@@ -309,7 +309,7 @@ public:
 	static TSharedRef<class FGlobalTabmanager> GetGlobalTabManager();
 
 	/** Initializes high dpi support for the process */
-	static void InitHighDPI();
+	static void InitHighDPI(const bool bForceEnable);
 
 	/** @return the root style node, which is the entry point to the style graph representing all the current style rules. */
 	const class FStyleNode* GetRootStyle() const;
