Updated Mobile PIE Device Specs for iOS and Android
Fixed issue with space in device spec name
Fixed issue with bezel calculation when in landscape

diff --git a/Engine/Source/Editor/PIEPreviewDeviceProfileSelector/Private/PIEPreviewDevice.cpp b/Engine/Source/Editor/PIEPreviewDeviceProfileSelector/Private/PIEPreviewDevice.cpp
--- a/Engine/Source/Editor/PIEPreviewDeviceProfileSelector/Private/PIEPreviewDevice.cpp
+++ b/Engine/Source/Editor/PIEPreviewDeviceProfileSelector/Private/PIEPreviewDevice.cpp
@@ -42,8 +42,13 @@ void FPIEPreviewDevice::ComputeViewportSize(const bool bClampWindowSize)
 	if (IsDeviceFlipped())
 	{
 		Swap(ScreenWidth, ScreenHeight);
-		Swap(ViewportRect.X, ViewportRect.Y);
 		Swap(ViewportRect.Width, ViewportRect.Height);
+		Swap(ViewportRect.X, ViewportRect.Y);
+		if (BezelTexture != nullptr)
+		{
+			// image is rotated counter-clockwise so the top coordinate is measured from what used to be the right edge.
+			ViewportRect.Y = BezelTexture->GetSizeX() - ViewportRect.Y - ViewportRect.Height;
+		}
 	}
 
 	float ScaleX = (float)ScreenWidth / (float)ViewportRect.Width;
diff --git a/Engine/Source/Editor/UnrealEd/Private/PlayLevel.cpp b/Engine/Source/Editor/UnrealEd/Private/PlayLevel.cpp
--- a/Engine/Source/Editor/UnrealEd/Private/PlayLevel.cpp
+++ b/Engine/Source/Editor/UnrealEd/Private/PlayLevel.cpp
@@ -1345,7 +1345,7 @@ void UEditorEngine::PlayStandaloneLocalPc(FString MapNameOverride, FIntPoint* Wi
 	{
 		if (PlayUsingMobilePreviewTargetDevice.IsEmpty() == false)
 		{
-			AdditionalParameters += TEXT(" -MobileTargetDevice=") + PlayUsingMobilePreviewTargetDevice;
+			AdditionalParameters += TEXT(" -MobileTargetDevice=\"") + PlayUsingMobilePreviewTargetDevice + TEXT("\"");
 		}
 		else
 		{
