Fixed issue where new Mobile PIE doesn't remember size setting from previous session

diff --git a/Engine/Source/Editor/PIEPreviewDeviceProfileSelector/Private/PIEPreviewDeviceProfileSelectorModule.cpp b/Engine/Source/Editor/PIEPreviewDeviceProfileSelector/Private/PIEPreviewDeviceProfileSelectorModule.cpp
--- a/Engine/Source/Editor/PIEPreviewDeviceProfileSelector/Private/PIEPreviewDeviceProfileSelectorModule.cpp
+++ b/Engine/Source/Editor/PIEPreviewDeviceProfileSelector/Private/PIEPreviewDeviceProfileSelectorModule.cpp
@@ -21,6 +21,7 @@
 #include "UnrealEngine.h"
 #include "Engine/GameViewportClient.h"
 #include "Engine/UserInterfaceSettings.h"
+#include "PIEPreviewSettings.h"
 
 DECLARE_LOG_CATEGORY_EXTERN(LogPIEPreviewDevice, Log, All); 
 DEFINE_LOG_CATEGORY(LogPIEPreviewDevice);
@@ -95,6 +96,7 @@ void FPIEPreviewDeviceModule::OnEngineInitComplete()
 		int32 TitleBarSize = SPIEPreviewWindow::GetDefaultTitleBarSize();
 		Device->SetupDevice(TitleBarSize);
 		
+		ReadWindowConfig();
 		WindowPtr->PrepareWindow(InitialWindowPosition, InitialWindowScaleValue, Device);
 		WindowPtr->ShowWindow();
 	}
@@ -102,31 +104,17 @@ void FPIEPreviewDeviceModule::OnEngineInitComplete()
 
 bool FPIEPreviewDeviceModule::ReadWindowConfig()
 {
-	InitialWindowScaleValue = 0.0f;
-	GConfig->GetFloat(TEXT("/Script/Engine.MobilePIE"), TEXT("DeviceScalingFactor"), InitialWindowScaleValue, GEngineIni);
+	auto* Settings = GetDefault<UPIEPreviewSettings>();
+	InitialWindowScaleValue = Settings->WindowScalingFactor;
 
-	// read window position
-	int32 WinX = 0, WinY = 0;
-	bool bFound = GConfig->GetInt(TEXT("/Script/Engine.MobilePIE"), TEXT("WindowPosX"), WinX, GEngineIni);
-	bFound &= GConfig->GetInt(TEXT("/Script/Engine.MobilePIE"), TEXT("WindowPosY"), WinY, GEngineIni);
-	InitialWindowPosition.Set(WinX, WinY);
+	InitialWindowPosition.Set(Settings->WindowPosX, Settings->WindowPosY);
 
-	return bFound;
+	return Settings->WindowPosX > 0 &&  Settings->WindowPosY > 0;
 }
 
 TSharedRef<SWindow> FPIEPreviewDeviceModule::CreatePIEPreviewDeviceWindow(FVector2D ClientSize, FText WindowTitle, EAutoCenter AutoCenterType, FVector2D ScreenPosition, TOptional<float> MaxWindowWidth, TOptional<float> MaxWindowHeight)
 {
-	bool bPosFound = ReadWindowConfig();
-
-	if (ScreenPosition.IsNearlyZero() && bPosFound)
-	{
-		ScreenPosition = InitialWindowPosition;
-		AutoCenterType = EAutoCenter::None;
-	}
-	else
-	{
-		InitialWindowPosition = ScreenPosition;
-	}
+	InitialWindowPosition = ScreenPosition;
 
 	FPIEPreviewWindowCoreStyle::InitializePIECoreStyle();
 
diff --git a/Engine/Source/Editor/PIEPreviewDeviceProfileSelector/Private/PIEPreviewSettings.cpp b/Engine/Source/Editor/PIEPreviewDeviceProfileSelector/Private/PIEPreviewSettings.cpp
new file mode 100644
--- /dev/null
+++ b/Engine/Source/Editor/PIEPreviewDeviceProfileSelector/Private/PIEPreviewSettings.cpp
@@ -0,0 +1,11 @@
+// Copyright 1998-2018 Epic Games, Inc. All Rights Reserved.
+
+#include "PIEPreviewSettings.h"
+
+UPIEPreviewSettings::UPIEPreviewSettings(const FObjectInitializer& ObjectInitializer)
+: Super(ObjectInitializer)
+{
+	WindowPosX = -1;
+	WindowPosY = -1;
+	WindowScalingFactor = 0;
+}
diff --git a/Engine/Source/Editor/PIEPreviewDeviceProfileSelector/Private/PIEPreviewSettings.h b/Engine/Source/Editor/PIEPreviewDeviceProfileSelector/Private/PIEPreviewSettings.h
new file mode 100644
--- /dev/null
+++ b/Engine/Source/Editor/PIEPreviewDeviceProfileSelector/Private/PIEPreviewSettings.h
@@ -0,0 +1,29 @@
+// Copyright 1998-2018 Epic Games, Inc. All Rights Reserved.
+
+/**
+*
+* A configuration class used by the PIE preview device system to save settings across sessions.
+*/
+
+#pragma once
+
+#include "CoreMinimal.h"
+#include "UObject/ObjectMacros.h"
+#include "UObject/Object.h"
+#include "PIEPreviewSettings.generated.h"
+
+UCLASS(hidecategories = Object, config = PIEPreviewSettings)
+class PIEPREVIEWDEVICEPROFILESELECTOR_API UPIEPreviewSettings : public UObject
+{
+	GENERATED_UCLASS_BODY()
+
+public:
+	UPROPERTY(config)
+	int32 WindowPosX;
+
+	UPROPERTY(config)
+	int32 WindowPosY;
+
+	UPROPERTY(config)
+	float WindowScalingFactor;
+};
\ No newline at end of file
diff --git a/Engine/Source/Editor/PIEPreviewDeviceProfileSelector/Private/PIEPreviewWindow.cpp b/Engine/Source/Editor/PIEPreviewDeviceProfileSelector/Private/PIEPreviewWindow.cpp
--- a/Engine/Source/Editor/PIEPreviewDeviceProfileSelector/Private/PIEPreviewWindow.cpp
+++ b/Engine/Source/Editor/PIEPreviewDeviceProfileSelector/Private/PIEPreviewWindow.cpp
@@ -16,6 +16,7 @@
 #include "Misc/ConfigCacheIni.h"
 #include "Slate/SGameLayerManager.h"
 #include "UnrealEngine.h"
+#include "PIEPreviewSettings.h"
 
 
 #if WITH_EDITOR
@@ -371,7 +372,9 @@ void SPIEPreviewWindow::SetWindowScaleFactor(const float ScaleFactor, const bool
 	// when required we will save the scaling value so it can be restored after session restart
 	if (bStore)
 	{
-		GConfig->SetFloat(TEXT("/Script/Engine.MobilePIE"), TEXT("DeviceScalingFactor"), WindowScalingFactor, GEngineIni);
+		auto* Settings = GetMutableDefault<UPIEPreviewSettings>();
+		Settings->WindowScalingFactor = ScaleFactor;
+		Settings->SaveConfig();
 	}
 
 	ScaleWindow(ScaleFactor);
@@ -422,8 +425,11 @@ void SPIEPreviewWindow::OnWindowMoved(const TSharedRef<SWindow>& Window)
 
 	// save the position so we can restore it if the session is restarted
 	FVector2D WindowPos = GetPositionInScreen();
-	GConfig->SetInt(TEXT("/Script/Engine.MobilePIE"), TEXT("WindowPosX"), FMath::CeilToInt(WindowPos.X), GEngineIni);
-	GConfig->SetInt(TEXT("/Script/Engine.MobilePIE"), TEXT("WindowPosY"), FMath::CeilToInt(WindowPos.Y), GEngineIni);
+
+	auto* Settings = GetMutableDefault<UPIEPreviewSettings>();
+	Settings->WindowPosX = FMath::CeilToInt(WindowPos.X);
+	Settings->WindowPosY = FMath::CeilToInt(WindowPos.Y);
+	Settings->SaveConfig();
 }
 
 void SPIEPreviewWindow::OnDisplayDPIChanged(TSharedRef<SWindow> Window)
diff --git a/Engine/Source/Editor/UnrealEd/Private/PlayLevel.cpp b/Engine/Source/Editor/UnrealEd/Private/PlayLevel.cpp
--- a/Engine/Source/Editor/UnrealEd/Private/PlayLevel.cpp
+++ b/Engine/Source/Editor/UnrealEd/Private/PlayLevel.cpp
@@ -1357,6 +1357,8 @@ void UEditorEngine::PlayStandaloneLocalPc(FString MapNameOverride, FIntPoint* Wi
 			AdditionalParameters += TEXT(" -opengl");
 		}
 		AdditionalParameters += TEXT(" -faketouches");
+
+		AdditionalParameters += TEXT(" -MultiprocessSaveConfig");
 	}
 
 	if (bPlayUsingVulkanPreview)
