Missed a file.

diff --git a/Engine/Source/Editor/PIEPreviewDeviceProfileSelector/Private/PIEPreviewDeviceProfileSelectorModule.cpp b/Engine/Source/Editor/PIEPreviewDeviceProfileSelector/Private/PIEPreviewDeviceProfileSelectorModule.cpp
--- a/Engine/Source/Editor/PIEPreviewDeviceProfileSelector/Private/PIEPreviewDeviceProfileSelectorModule.cpp
+++ b/Engine/Source/Editor/PIEPreviewDeviceProfileSelector/Private/PIEPreviewDeviceProfileSelectorModule.cpp
@@ -222,7 +222,7 @@ void FPIEPreviewDeviceModule::OnViewportCreated()
 	// disable mouse viewport locking
 	if (GEngine->GameViewport != nullptr)
 	{
-		GEngine->GameViewport->SetCaptureMouseOnClick(EMouseCaptureMode::NoCapture);
+		GEngine->GameViewport->SetMouseCaptureMode(EMouseCaptureMode::NoCapture);
 		GEngine->GameViewport->SetMouseLockMode(EMouseLockMode::DoNotLock);
 	}
 }
