Add API to set PIE preview device from code.
Allows use of PIE device profile selector from editor preview.

diff --git a/Engine/Source/Editor/PIEPreviewDeviceProfileSelector/Private/PIEPreviewDevice.cpp b/Engine/Source/Editor/PIEPreviewDeviceProfileSelector/Private/PIEPreviewDevice.cpp
--- a/Engine/Source/Editor/PIEPreviewDeviceProfileSelector/Private/PIEPreviewDevice.cpp
+++ b/Engine/Source/Editor/PIEPreviewDeviceProfileSelector/Private/PIEPreviewDevice.cpp
@@ -411,6 +411,24 @@ void FPIEPreviewDevice::ApplyRHIOverrides() const
 	}
 }
 
+bool FPIEPreviewDevice::GetSelectorPropertyValue(const FName& PropertyType, FString& PropertyValueOUT) const
+{
+	switch (DeviceSpecs->DevicePlatform)
+	{
+		case EPIEPreviewDeviceType::Android:
+		{
+			IDeviceProfileSelectorModule* AndroidDeviceProfileSelector = FModuleManager::LoadModulePtr<IDeviceProfileSelectorModule>("AndroidDeviceProfileSelector");
+			if (AndroidDeviceProfileSelector)
+			{
+				return AndroidDeviceProfileSelector->GetSelectorPropertyValue(PropertyType, PropertyValueOUT);
+			}
+		}
+		break;
+	}
+	return false;
+}
+
+
 FString FPIEPreviewDevice::GetProfile() const
 {
 	FString Profile;
diff --git a/Engine/Source/Editor/PIEPreviewDeviceProfileSelector/Private/PIEPreviewDevice.h b/Engine/Source/Editor/PIEPreviewDeviceProfileSelector/Private/PIEPreviewDevice.h
--- a/Engine/Source/Editor/PIEPreviewDeviceProfileSelector/Private/PIEPreviewDevice.h
+++ b/Engine/Source/Editor/PIEPreviewDeviceProfileSelector/Private/PIEPreviewDevice.h
@@ -117,6 +117,8 @@ public:
 
 	FString GetProfile() const;
 
+	bool GetSelectorPropertyValue(const FName& PropertyType, FString& PropertyValueOUT) const;
+
 	UTexture2D *GetBezelTexture();
 };
 
diff --git a/Engine/Source/Editor/PIEPreviewDeviceProfileSelector/Private/PIEPreviewDeviceProfileSelectorModule.cpp b/Engine/Source/Editor/PIEPreviewDeviceProfileSelector/Private/PIEPreviewDeviceProfileSelectorModule.cpp
--- a/Engine/Source/Editor/PIEPreviewDeviceProfileSelector/Private/PIEPreviewDeviceProfileSelectorModule.cpp
+++ b/Engine/Source/Editor/PIEPreviewDeviceProfileSelector/Private/PIEPreviewDeviceProfileSelectorModule.cpp
@@ -31,17 +31,10 @@ IMPLEMENT_MODULE(FPIEPreviewDeviceModule, PIEPreviewDeviceProfileSelector);
 void FPIEPreviewDeviceModule::StartupModule()
 {
 	// Parse the json file specified on the command line
-
-	if (FParse::Value(FCommandLine::Get(), GetPreviewDeviceCommandSwitch(), PreviewDevice))
+	FString PreviewDeviceCmdLine;
+	if (FParse::Value(FCommandLine::Get(), GetPreviewDeviceCommandSwitch(), PreviewDeviceCmdLine))
 	{
-		const FString Filename = FindDeviceSpecificationFilePath(PreviewDevice);
-
-		FString Json;
-		if (FFileHelper::LoadFileToString(Json, *Filename))
-		{
-			TSharedRef<TJsonReader<> > JsonReader = TJsonReaderFactory<>::Create(Json);
-			FJsonSerializer::Deserialize(JsonReader, JsonRootObject);
-		}
+		SetPreviewDevice(PreviewDeviceCmdLine);
 	}
 }
 
@@ -70,6 +63,21 @@ void FPIEPreviewDeviceModule::ShutdownModule()
 	}
 }
 
+void FPIEPreviewDeviceModule::SetPreviewDevice(const FString& DeviceName)
+{
+	PreviewDevice = DeviceName;
+	const FString Filename = FindDeviceSpecificationFilePath(PreviewDevice);
+
+	FString Json;
+	if (FFileHelper::LoadFileToString(Json, *Filename))
+	{
+		TSharedRef<TJsonReader<> > JsonReader = TJsonReaderFactory<>::Create(Json);
+		FJsonSerializer::Deserialize(JsonReader, JsonRootObject);
+	}
+
+	InitPreviewDevice();
+}
+
 void FPIEPreviewDeviceModule::ApplyCommandLineOverrides()
 {
 	// Here we need to parse the json directly as we have not yet initialized the UObject system
@@ -93,6 +101,11 @@ FString const FPIEPreviewDeviceModule::GetRuntimeDeviceProfileName()
 	return DeviceProfile;
 }
 
+bool FPIEPreviewDeviceModule::GetSelectorPropertyValue(const FName& PropertyType, FString& PropertyValueOUT)
+{ 
+	return Device->GetSelectorPropertyValue(PropertyType, PropertyValueOUT);
+}
+
 void FPIEPreviewDeviceModule::InitPreviewDevice()
 {
 	bInitialized = true;
@@ -106,7 +119,11 @@ void FPIEPreviewDeviceModule::InitPreviewDevice()
 	bool bReadSuccess = ReadDeviceSpecification();
 	checkf(bReadSuccess, TEXT("Unable to read PIE Preview Device specification"));
 
-	Device->ApplyRHIPrerequisitesOverrides();
+	// Do not apply RHI overrides when using the editor
+	if(!GIsEditor)
+	{
+		Device->ApplyRHIPrerequisitesOverrides();
+	}
 	DeviceProfile = Device->GetProfile();
 }
 
diff --git a/Engine/Source/Editor/PIEPreviewDeviceProfileSelector/Public/PIEPreviewDeviceProfileSelectorModule.h b/Engine/Source/Editor/PIEPreviewDeviceProfileSelector/Public/PIEPreviewDeviceProfileSelectorModule.h
--- a/Engine/Source/Editor/PIEPreviewDeviceProfileSelector/Public/PIEPreviewDeviceProfileSelectorModule.h
+++ b/Engine/Source/Editor/PIEPreviewDeviceProfileSelector/Public/PIEPreviewDeviceProfileSelectorModule.h
@@ -23,6 +23,7 @@ public:
 
 	//~ Begin IDeviceProfileSelectorModule Interface
 	virtual const FString GetRuntimeDeviceProfileName() override;
+	virtual bool GetSelectorPropertyValue(const FName& PropertyType, FString& PropertyValueOUT) override;
 	//~ End IDeviceProfileSelectorModule Interface
 
 	//~ Begin IModuleInterface Interface
@@ -61,6 +62,8 @@ public:
 	DECLARE_MULTICAST_DELEGATE_TwoParams(FAddToDevicePreviewMenu, const FText& DeviceCategoryName, class FMenuBuilder& MenuBuilder);
 	FAddToDevicePreviewMenu AddToDevicePreviewMenuDelegates;
 
+	virtual void SetPreviewDevice(const FString& DeviceName) override;
+
 private:
 	static const TCHAR* GetPreviewDeviceCommandSwitch()
 	{
diff --git a/Engine/Source/Runtime/Engine/Public/IPIEPreviewDeviceModule.h b/Engine/Source/Runtime/Engine/Public/IPIEPreviewDeviceModule.h
--- a/Engine/Source/Runtime/Engine/Public/IPIEPreviewDeviceModule.h
+++ b/Engine/Source/Runtime/Engine/Public/IPIEPreviewDeviceModule.h
@@ -53,5 +53,6 @@ class IPIEPreviewDeviceModule : public IDeviceProfileSelectorModule
 		{
 		}
 
+		virtual void SetPreviewDevice(const FString& DeviceName) = 0;
 		//~ End IDeviceProfileSelectorModule Interface
 };
\ No newline at end of file
