Change  Default to Mobile Forward

diff --git a/Engine/Source/Editor/PIEPreviewDeviceProfileSelector/Private/PIEPreviewDevice.cpp b/Engine/Source/Editor/PIEPreviewDeviceProfileSelector/Private/PIEPreviewDevice.cpp
--- a/Engine/Source/Editor/PIEPreviewDeviceProfileSelector/Private/PIEPreviewDevice.cpp
+++ b/Engine/Source/Editor/PIEPreviewDeviceProfileSelector/Private/PIEPreviewDevice.cpp
@@ -272,16 +272,16 @@ ERHIFeatureLevel::Type FPIEPreviewDevice::GetPreviewDeviceFeatureLevel() const
 			// use a local ini file, so that Windows can read settings that are in SwitchEngine.ini files...
 			FConfigFile SwitchSettings;
 			FConfigCacheIni::LoadLocalIniFile(SwitchSettings, TEXT("Engine"), true, TEXT("Switch"));
-			bool bUseMobileRenderer = false;
-			SwitchSettings.GetBool(TEXT("/Script/SwitchRuntimeSettings.SwitchRuntimeSettings"), TEXT("bUseMobileForwardRenderer"), bUseMobileRenderer);
+			bool bSupportDesktopRenderer = false;
+			SwitchSettings.GetBool(TEXT("/Script/SwitchRuntimeSettings.SwitchRuntimeSettings"), TEXT("bSupportDesktopRenderer"), bSupportDesktopRenderer);
 
-			if (bUseMobileRenderer)
+			if (bSupportDesktopRenderer)
 			{
-				return ERHIFeatureLevel::ES3_1;
+				return ERHIFeatureLevel::SM5;
 			}
 			else
 			{
-				return ERHIFeatureLevel::SM5;
+				return ERHIFeatureLevel::ES3_1;
 			}
 		}
 	}
@@ -476,15 +476,15 @@ FString FPIEPreviewDevice::GetProfile() const
 			// load switch renderer configuration
 			FConfigFile SwitchSettings;
 			FConfigCacheIni::LoadLocalIniFile(SwitchSettings, TEXT("Engine"), true, TEXT("Switch"));
-			bool bUseMobileRenderer = false;
-			SwitchSettings.GetBool(TEXT("/Script/SwitchRuntimeSettings.SwitchRuntimeSettings"), TEXT("bUseMobileForwardRenderer"), bUseMobileRenderer);
+			bool bSupportDesktopRenderer = false;
+			SwitchSettings.GetBool(TEXT("/Script/SwitchRuntimeSettings.SwitchRuntimeSettings"), TEXT("bSupportDesktopRenderer"), bSupportDesktopRenderer);
 			bool bForwardShading = false;
 			SwitchSettings.GetBool(TEXT("/Script/SwitchRuntimeSettings.SwitchRuntimeSettings"), TEXT("bUseForwardShading"), bForwardShading);
 	
 			// taken from FSwitchApplication::UpdateActiveDeviceProfile()
 			Profile = DeviceSpecs->SwitchProperties.Docked ?
-				((bUseMobileRenderer || bForwardShading) ? TEXT("Switch_Console_Forward") : TEXT("Switch_Console_Deferred")) :
-				((bUseMobileRenderer || bForwardShading) ? TEXT("Switch_Handheld_Forward") : TEXT("Switch_Handheld_Deferred"));
+				((!bSupportDesktopRenderer || bForwardShading) ? TEXT("Switch_Console_Forward") : TEXT("Switch_Console_Deferred")) :
+				((!bSupportDesktopRenderer || bForwardShading) ? TEXT("Switch_Handheld_Forward") : TEXT("Switch_Handheld_Deferred"));
 		
 			break;
 		}
