#ue5 - Truncation/conversion fixes for ProjectLauncher, MediaAssets, ProjectTargetPlatformEditor, and associated UI code.

diff --git a/Engine/Source/Developer/ProjectLauncher/Private/Widgets/Archive/SProjectLauncherArchivePage.cpp b/Engine/Source/Developer/ProjectLauncher/Private/Widgets/Archive/SProjectLauncherArchivePage.cpp
--- a/Engine/Source/Developer/ProjectLauncher/Private/Widgets/Archive/SProjectLauncherArchivePage.cpp
+++ b/Engine/Source/Developer/ProjectLauncher/Private/Widgets/Archive/SProjectLauncherArchivePage.cpp
@@ -64,7 +64,7 @@ void SProjectLauncherArchivePage::Construct(const FArguments& InArgs, const TSha
 			.Padding(0.0, 8.0, 0.0, 0.0)
 			[
 				SNew(SBorder)
-					.Padding(8.0)
+					.Padding(8.0f)
 					.BorderImage(FAppStyle::GetBrush("ToolPanel.GroupBorder"))
 					.Visibility(this, &SProjectLauncherArchivePage::HandleArchiveVisibility)
 					[
diff --git a/Engine/Source/Developer/ProjectLauncher/Private/Widgets/Build/SProjectLauncherBuildPage.cpp b/Engine/Source/Developer/ProjectLauncher/Private/Widgets/Build/SProjectLauncherBuildPage.cpp
--- a/Engine/Source/Developer/ProjectLauncher/Private/Widgets/Build/SProjectLauncherBuildPage.cpp
+++ b/Engine/Source/Developer/ProjectLauncher/Private/Widgets/Build/SProjectLauncherBuildPage.cpp
@@ -136,7 +136,7 @@ void SProjectLauncherBuildPage::Construct(const FArguments& InArgs, const TShare
                 SNew(SExpandableArea)
 					.AreaTitle(LOCTEXT("AdvancedAreaTitle", "Advanced Settings"))
 					.InitiallyCollapsed(true)
-					.Padding(8.0)
+					.Padding(8.0f)
 					.BodyContent()
 					[
 						SNew(SVerticalBox)
diff --git a/Engine/Source/Developer/ProjectLauncher/Private/Widgets/Cook/SProjectLauncherCookOnTheFlySettings.cpp b/Engine/Source/Developer/ProjectLauncher/Private/Widgets/Cook/SProjectLauncherCookOnTheFlySettings.cpp
--- a/Engine/Source/Developer/ProjectLauncher/Private/Widgets/Cook/SProjectLauncherCookOnTheFlySettings.cpp
+++ b/Engine/Source/Developer/ProjectLauncher/Private/Widgets/Cook/SProjectLauncherCookOnTheFlySettings.cpp
@@ -47,7 +47,7 @@ void SProjectLauncherCookOnTheFlySettings::Construct(const FArguments& InArgs, c
 				SNew(SExpandableArea)
 					.AreaTitle(LOCTEXT("AdvancedAreaTitle", "Advanced Settings"))
 					.InitiallyCollapsed(true)
-					.Padding(8.0)
+					.Padding(8.0f)
 					.BodyContent()
 					[
 						SNew(SVerticalBox)
diff --git a/Engine/Source/Developer/ProjectLauncher/Private/Widgets/Deploy/SProjectLauncherDeployRepositorySettings.cpp b/Engine/Source/Developer/ProjectLauncher/Private/Widgets/Deploy/SProjectLauncherDeployRepositorySettings.cpp
--- a/Engine/Source/Developer/ProjectLauncher/Private/Widgets/Deploy/SProjectLauncherDeployRepositorySettings.cpp
+++ b/Engine/Source/Developer/ProjectLauncher/Private/Widgets/Deploy/SProjectLauncherDeployRepositorySettings.cpp
@@ -32,7 +32,7 @@ void SProjectLauncherDeployRepositorySettings::Construct(const FArguments& InArg
 		.FillHeight(1.0)
 		[
 			SNew(SBorder)
-				.Padding(8.0)
+				.Padding(8.0f)
 				.BorderImage(FAppStyle::GetBrush("ToolPanel.GroupBorder"))
 				[
 					SNew(SVerticalBox)
diff --git a/Engine/Source/Developer/ProjectLauncher/Private/Widgets/Deploy/SProjectLauncherDeployTargetListRow.h b/Engine/Source/Developer/ProjectLauncher/Private/Widgets/Deploy/SProjectLauncherDeployTargetListRow.h
--- a/Engine/Source/Developer/ProjectLauncher/Private/Widgets/Deploy/SProjectLauncherDeployTargetListRow.h
+++ b/Engine/Source/Developer/ProjectLauncher/Private/Widgets/Deploy/SProjectLauncherDeployTargetListRow.h
@@ -114,8 +114,8 @@ public:
 				.AutoWidth()
 				[
 					SNew(SBox)
-						.WidthOverride(24)
-						.HeightOverride(24)
+						.WidthOverride(24.f)
+						.HeightOverride(24.f)
 						[
 							SNew(SImage)
 								.Image(this, &SProjectLauncherDeployTargetListRow::HandleDeviceImage)
diff --git a/Engine/Source/Developer/ProjectLauncher/Private/Widgets/Deploy/SProjectLauncherSimpleDeviceListRow.h b/Engine/Source/Developer/ProjectLauncher/Private/Widgets/Deploy/SProjectLauncherSimpleDeviceListRow.h
--- a/Engine/Source/Developer/ProjectLauncher/Private/Widgets/Deploy/SProjectLauncherSimpleDeviceListRow.h
+++ b/Engine/Source/Developer/ProjectLauncher/Private/Widgets/Deploy/SProjectLauncherSimpleDeviceListRow.h
@@ -117,8 +117,8 @@ public:
 					.Padding(14, 0, 12, 0)
 					[
 						SNew(SBox)
-						.WidthOverride(44)
-						.HeightOverride(44)
+						.WidthOverride(44.f)
+						.HeightOverride(44.f)
 						[
 							SNew(SImage)
 							.Image(this, &SProjectLauncherSimpleDeviceListRow::HandleDeviceImage)
diff --git a/Engine/Source/Developer/ProjectLauncher/Private/Widgets/Launch/SProjectLauncherLaunchPage.cpp b/Engine/Source/Developer/ProjectLauncher/Private/Widgets/Launch/SProjectLauncherLaunchPage.cpp
--- a/Engine/Source/Developer/ProjectLauncher/Private/Widgets/Launch/SProjectLauncherLaunchPage.cpp
+++ b/Engine/Source/Developer/ProjectLauncher/Private/Widgets/Launch/SProjectLauncherLaunchPage.cpp
@@ -96,7 +96,7 @@ void SProjectLauncherLaunchPage::Construct(const FArguments& InArgs, const TShar
 			[
 				SNew(SBorder)
 					.BorderImage(FAppStyle::GetBrush("ToolPanel.GroupBorder"))
-					.Padding(8.0)
+					.Padding(8.0f)
 					.Visibility(this, &SProjectLauncherLaunchPage::HandleValidationErrorIconVisibility, ELauncherProfileValidationErrors::CustomRolesNotSupportedYet)
 					[
 						SNew(SHorizontalBox)
@@ -126,7 +126,7 @@ void SProjectLauncherLaunchPage::Construct(const FArguments& InArgs, const TShar
 				SNew(SExpandableArea)
 					.AreaTitle(LOCTEXT("DefaultRoleAreaTitle", "Default Role"))
 					.InitiallyCollapsed(false)
-					.Padding(8.0)
+					.Padding(8.0f)
 					.Visibility(this, &SProjectLauncherLaunchPage::HandleLaunchSettingsVisibility)
 					.BodyContent()
 					[
diff --git a/Engine/Source/Developer/ProjectLauncher/Private/Widgets/Progress/SProjectLauncherProgress.h b/Engine/Source/Developer/ProjectLauncher/Private/Widgets/Progress/SProjectLauncherProgress.h
--- a/Engine/Source/Developer/ProjectLauncher/Private/Widgets/Progress/SProjectLauncherProgress.h
+++ b/Engine/Source/Developer/ProjectLauncher/Private/Widgets/Progress/SProjectLauncherProgress.h
@@ -204,11 +204,11 @@ public:
 									// Cache the max desired width seen so far of rows in the message list.
 									// We use this to prevent the horizontal scroll bar from constantly shrinking/expanding as rows are virtualized.
 									// This will get reset whenever the message list is cleared.
-									MaxMessageListRowWidth = FMath::Max(MaxMessageListRowWidth, MessageListView->GetDesiredSize().X);
+									MaxMessageListRowWidth = FMath::Max(MaxMessageListRowWidth, (float)MessageListView->GetDesiredSize().X);
 
 									// Make the message list at least as wide as the scroll box so that there's not an empty space to its right;
 									// scroll boxes do not deal with alignment in the scroll direction (horizontal in this case).
-									const float CurrentScrollWidth = HorizontalScrollBox->GetTickSpaceGeometry().GetLocalSize().X;
+									const float CurrentScrollWidth = (float)HorizontalScrollBox->GetTickSpaceGeometry().GetLocalSize().X;
 									return FMath::Max(MaxMessageListRowWidth, CurrentScrollWidth);
 								})
 								[
diff --git a/Engine/Source/Developer/ProjectLauncher/Private/Widgets/Shared/SProjectLauncherProfileNameDescEditor.cpp b/Engine/Source/Developer/ProjectLauncher/Private/Widgets/Shared/SProjectLauncherProfileNameDescEditor.cpp
--- a/Engine/Source/Developer/ProjectLauncher/Private/Widgets/Shared/SProjectLauncherProfileNameDescEditor.cpp
+++ b/Engine/Source/Developer/ProjectLauncher/Private/Widgets/Shared/SProjectLauncherProfileNameDescEditor.cpp
@@ -29,8 +29,8 @@ void SProjectLauncherProfileNameDescEditor::Construct(const FArguments& InArgs,
 			.VAlign(VAlign_Center)
 			[
 				SNew(SBox)
-				.WidthOverride(20)
-				.HeightOverride(20)
+				.WidthOverride(20.f)
+				.HeightOverride(20.f)
 				[
 					SNew(SImage)
 					.Image(this, &SProjectLauncherProfileNameDescEditor::HandleProfileImage)
diff --git a/Engine/Source/Developer/ProjectLauncher/Private/Widgets/Shared/SProjectLauncherValidation.h b/Engine/Source/Developer/ProjectLauncher/Private/Widgets/Shared/SProjectLauncherValidation.h
--- a/Engine/Source/Developer/ProjectLauncher/Private/Widgets/Shared/SProjectLauncherValidation.h
+++ b/Engine/Source/Developer/ProjectLauncher/Private/Widgets/Shared/SProjectLauncherValidation.h
@@ -212,7 +212,7 @@ protected:
 
 		+ SHorizontalBox::Slot()
 			.AutoWidth()
-			.Padding(2.0)
+			.Padding(2.0f)
 			[
 				SNew(SImage)
 					.Image(FAppStyle::GetBrush(IconName))
@@ -240,7 +240,7 @@ protected:
 
 		+ SHorizontalBox::Slot()
 			.AutoWidth()
-			.Padding(2.0)
+			.Padding(2.0f)
 			[
 				SNew(SImage)
 					.Image(FAppStyle::GetBrush(IconName))
diff --git a/Engine/Source/Editor/PIEPreviewDeviceProfileSelector/Private/PIEPreviewDevice.cpp b/Engine/Source/Editor/PIEPreviewDeviceProfileSelector/Private/PIEPreviewDevice.cpp
--- a/Engine/Source/Editor/PIEPreviewDeviceProfileSelector/Private/PIEPreviewDevice.cpp
+++ b/Engine/Source/Editor/PIEPreviewDeviceProfileSelector/Private/PIEPreviewDevice.cpp
@@ -58,22 +58,22 @@ void FPIEPreviewDevice::ComputeViewportSize(const bool bClampWindowSize)
 
  	// compute widow size
 	WindowWidth = ScreenWidth;
-	WindowHeight = ScreenHeight + WindowTitleBarSize * DPIScaleFactor;
+	WindowHeight = ScreenHeight + FMath::TruncToInt32((float)WindowTitleBarSize * DPIScaleFactor);
 
 	// compute viewport margin
 	if (bShowBezel && BezelTexture != nullptr)
 	{
 		// compute widow size
-		WindowWidth += 2.0f * ViewportRect.X * ScaleX;
-		WindowHeight += 2.0f * ViewportRect.Y * ScaleY;
+		WindowWidth += FMath::TruncToInt32(2.0f * ViewportRect.X * ScaleX);
+		WindowHeight += FMath::TruncToInt32(2.0f * ViewportRect.Y * ScaleY);
 
 		ViewportRect.X = FMath::RoundToInt(ViewportRect.X * BezelScaleFactor);
 		ViewportRect.Y = FMath::RoundToInt(ViewportRect.Y * BezelScaleFactor);
 		ViewportRect.Width = FMath::RoundToInt(ViewportRect.Width * BezelScaleFactor);
 		ViewportRect.Height = FMath::RoundToInt(ViewportRect.Height * BezelScaleFactor);
 
-		ViewportMargin.Left = ViewportRect.X;
-		ViewportMargin.Top = ViewportRect.Y;
+		ViewportMargin.Left = (float)ViewportRect.X;
+		ViewportMargin.Top = (float)ViewportRect.Y;
 
 		int32 BezelWidth = IsDeviceFlipped() ? BezelTexture->GetSizeY() : BezelTexture->GetSizeX();
 		int32 BezelHeight = IsDeviceFlipped() ? BezelTexture->GetSizeX() : BezelTexture->GetSizeY();
@@ -107,7 +107,7 @@ void FPIEPreviewDevice::ComputeViewportSize(const bool bClampWindowSize)
 			ScaleY *= ScaleFactor;
 
 			WindowWidth = DesktopWidth;
-			WindowHeight *= ScaleFactor;
+			WindowHeight = FMath::TruncToInt32((float)WindowHeight * ScaleFactor);
 
 			ViewportMargin = ViewportMargin * ScaleFactor;
 		}
@@ -117,7 +117,7 @@ void FPIEPreviewDevice::ComputeViewportSize(const bool bClampWindowSize)
 			ScaleX *= ScaleFactor;
 			ScaleY *= ScaleFactor;
 
-			WindowWidth *= ScaleFactor;
+			WindowWidth = FMath::TruncToInt32((float)WindowWidth * ScaleFactor);
 			WindowHeight = DesktopHeight;
 
 			ViewportMargin = ViewportMargin * ScaleFactor;
@@ -162,8 +162,8 @@ void FPIEPreviewDevice::ComputeContentScaledResolution(int32& Width, int32& Heig
 					RequestedContentScaleFactor = DeviceSpecs->IOSProperties.NativeScaleFactor;
 				}
 
-				Width *= RequestedContentScaleFactor;
-				Height *= RequestedContentScaleFactor;
+				Width = FMath::TruncToInt32((float)Width * RequestedContentScaleFactor);
+				Height = FMath::TruncToInt32((float)Height * RequestedContentScaleFactor);
 			}
 			break;
 
@@ -177,8 +177,8 @@ void FPIEPreviewDevice::ComputeDeviceResolution(int32& Width, int32& Height)
 {
 	ComputeContentScaledResolution(Width, Height);
 
-	Width *= ResolutionScaleFactor;
-	Height *= ResolutionScaleFactor;
+	Width = FMath::TruncToInt32((float)Width * ResolutionScaleFactor);
+	Height = FMath::TruncToInt32((float)Height * ResolutionScaleFactor);
 }
 
 void FPIEPreviewDevice::DetermineScreenOrientationRequirements(bool& bNeedPortrait, bool& bNeedLandscape)
@@ -496,5 +496,5 @@ FString FPIEPreviewDevice::GetProfile() const
 
 int32 FPIEPreviewDevice::GetWindowClientHeight() const
 {
-	return WindowHeight - WindowTitleBarSize * DPIScaleFactor;
+	return WindowHeight - FMath::TruncToInt32((float)WindowTitleBarSize * DPIScaleFactor);
 }
diff --git a/Engine/Source/Editor/PIEPreviewDeviceProfileSelector/Private/PIEPreviewDeviceProfileSelectorModule.cpp b/Engine/Source/Editor/PIEPreviewDeviceProfileSelector/Private/PIEPreviewDeviceProfileSelectorModule.cpp
--- a/Engine/Source/Editor/PIEPreviewDeviceProfileSelector/Private/PIEPreviewDeviceProfileSelectorModule.cpp
+++ b/Engine/Source/Editor/PIEPreviewDeviceProfileSelector/Private/PIEPreviewDeviceProfileSelectorModule.cpp
@@ -206,7 +206,7 @@ void FPIEPreviewDeviceModule::UpdateDisplayResolution()
 	}
 
 	const int32 ClientWidth = Device->GetWindowWidth();
-	const int32 ClientHeight = Device->GetWindowHeight() - WindowPtr->GetTitleBarSize().Get();
+	const int32 ClientHeight = Device->GetWindowHeight() - FMath::TruncToInt32(WindowPtr->GetTitleBarSize().Get());
 
 	FSystemResolution::RequestResolutionChange(ClientWidth, ClientHeight, EWindowMode::Windowed);
 	IConsoleManager::Get().CallAllConsoleVariableSinks();
diff --git a/Engine/Source/Editor/PIEPreviewDeviceProfileSelector/Private/PIEPreviewWindow.cpp b/Engine/Source/Editor/PIEPreviewDeviceProfileSelector/Private/PIEPreviewWindow.cpp
--- a/Engine/Source/Editor/PIEPreviewDeviceProfileSelector/Private/PIEPreviewWindow.cpp
+++ b/Engine/Source/Editor/PIEPreviewDeviceProfileSelector/Private/PIEPreviewWindow.cpp
@@ -50,8 +50,8 @@ void SPIEPreviewWindow::ComputeBezelOrientation()
 {
 	if (BezelImage.IsValid())
 	{
- 		float Width = Device->GetWindowClientWidth();
- 		float Height = Device->GetWindowClientHeight();
+ 		float Width = (float)Device->GetWindowClientWidth();
+ 		float Height = (float)Device->GetWindowClientHeight();
 
 		bool bBezelRotated = Device->IsDeviceFlipped();
 
@@ -322,8 +322,8 @@ void SPIEPreviewWindow::CreatePIEPreviewBezelOverlay(UTexture2D* pBezelImage)
 
 void SPIEPreviewWindow::ValidatePosition(FVector2D& WindowPos)
 {
-	WindowPos.X = FMath::CeilToInt(WindowPos.X);
-	WindowPos.Y = FMath::CeilToInt(WindowPos.Y);
+	WindowPos.X = FMath::CeilToInt32(WindowPos.X);
+	WindowPos.Y = FMath::CeilToInt32(WindowPos.Y);
 
 	FDisplayMetrics DisplayMetrics;
 	FDisplayMetrics::RebuildDisplayMetrics(DisplayMetrics);
@@ -353,7 +353,7 @@ void SPIEPreviewWindow::PrepareWindow(FVector2D WindowPosition, const float Init
 
 	// update display resolution
 	const int32 ClientWidth = Device->GetWindowWidth();
-	const int32 ClientHeight = Device->GetWindowHeight() - GetTitleBarSize().Get();
+	const int32 ClientHeight = Device->GetWindowHeight() - FMath::TruncToInt32(GetTitleBarSize().Get());
 	FSystemResolution::RequestResolutionChange(ClientWidth, ClientHeight, EWindowMode::Windowed);
 	IConsoleManager::Get().CallAllConsoleVariableSinks();
 
@@ -428,8 +428,8 @@ void SPIEPreviewWindow::OnWindowMoved(const TSharedRef<SWindow>& Window)
 	FVector2D WindowPos = GetPositionInScreen();
 
 	auto* Settings = GetMutableDefault<UPIEPreviewSettings>();
-	Settings->WindowPosX = FMath::CeilToInt(WindowPos.X);
-	Settings->WindowPosY = FMath::CeilToInt(WindowPos.Y);
+	Settings->WindowPosX = FMath::CeilToInt32(WindowPos.X);
+	Settings->WindowPosY = FMath::CeilToInt32(WindowPos.Y);
 	Settings->SaveConfig();
 }
 
@@ -442,10 +442,10 @@ void SPIEPreviewWindow::OnDisplayDPIChanged(TSharedRef<SWindow> Window)
 float SPIEPreviewWindow::ComputeDPIScaleFactor()
 {
 	FVector2D WindowPos = GetPositionInScreen();
-	int32 PointX = FMath::RoundToInt(WindowPos.X);
-	int32 PointY = FMath::RoundToInt(WindowPos.Y);
+	int32 PointX = FMath::RoundToInt32(WindowPos.X);
+	int32 PointY = FMath::RoundToInt32(WindowPos.Y);
 
-	float DPIFactor = FPlatformApplicationMisc::GetDPIScaleFactorAtPoint(PointX, PointY);
+	float DPIFactor = FPlatformApplicationMisc::GetDPIScaleFactorAtPoint((float)PointX, (float)PointY);
 
 	return DPIFactor;
 }
@@ -458,8 +458,8 @@ float SPIEPreviewWindow::ComputeScaleToDeviceSizeFactor() const
 	FDisplayMetrics::RebuildDisplayMetrics(DisplayMetrics);
 
 	FVector2D WindowPos = GetPositionInScreen();
-	int32 PointX = FMath::RoundToInt(WindowPos.X);
-	int32 PointY = FMath::RoundToInt(WindowPos.Y);
+	int32 PointX = FMath::RoundToInt32(WindowPos.X);
+	int32 PointY = FMath::RoundToInt32(WindowPos.Y);
 
 	FPlatformRect& VirtualDisplayRect = DisplayMetrics.VirtualDisplayRect;
 	PointX = FMath::Clamp(PointX, VirtualDisplayRect.Left, VirtualDisplayRect.Right);
@@ -566,7 +566,7 @@ void SPIEPreviewWindow::UpdateWindow()
 
 	if (PosY + Device->GetWindowHeight() > DisplayMetrics.VirtualDisplayRect.Bottom)
 	{
-		PosY = DisplayMetrics.VirtualDisplayRect.Bottom - Device->GetWindowHeight() - SWindowDefs::DefaultTitleBarSize;
+		PosY = DisplayMetrics.VirtualDisplayRect.Bottom - Device->GetWindowHeight() - FMath::TruncToInt32(SWindowDefs::DefaultTitleBarSize);
 	}
 	PosY = FMath::Max(DisplayMetrics.VirtualDisplayRect.Top, PosY);
 
@@ -614,7 +614,7 @@ void SPIEPreviewWindow::PrepareShutdown()
 
 int32 SPIEPreviewWindow::GetDefaultTitleBarSize()
 {
-	return SWindowDefs::DefaultTitleBarSize;
+	return FMath::TruncToInt32(SWindowDefs::DefaultTitleBarSize);
 }
 
 void SPIEPreviewWindow::SetGameLayerManagerWidget(TSharedPtr<class SGameLayerManager> GameLayerManager)
diff --git a/Engine/Source/Editor/PIEPreviewDeviceProfileSelector/Private/PIEPreviewWindowTitleBar.cpp b/Engine/Source/Editor/PIEPreviewDeviceProfileSelector/Private/PIEPreviewWindowTitleBar.cpp
--- a/Engine/Source/Editor/PIEPreviewDeviceProfileSelector/Private/PIEPreviewWindowTitleBar.cpp
+++ b/Engine/Source/Editor/PIEPreviewDeviceProfileSelector/Private/PIEPreviewWindowTitleBar.cpp
@@ -100,7 +100,7 @@ void SPIEPreviewWindowTitleBar::MakeTitleBarContentWidgets(TSharedPtr< SWidget >
 	ScreenRotationButton = SNew(SButton)
 		.IsFocusable(false)
 		.IsEnabled(true)
-		.ContentPadding(0)
+		.ContentPadding(0.f)
 		.OnClicked_Lambda
 		(
 			[this]()
diff --git a/Engine/Source/Editor/ProjectTargetPlatformEditor/Private/Widgets/SProjectTargetPlatformSettings.cpp b/Engine/Source/Editor/ProjectTargetPlatformEditor/Private/Widgets/SProjectTargetPlatformSettings.cpp
--- a/Engine/Source/Editor/ProjectTargetPlatformEditor/Private/Widgets/SProjectTargetPlatformSettings.cpp
+++ b/Engine/Source/Editor/ProjectTargetPlatformEditor/Private/Widgets/SProjectTargetPlatformSettings.cpp
@@ -115,8 +115,8 @@ TSharedRef<SWidget> SProjectTargetPlatformSettings::MakePlatformRow(const FText&
 		.VAlign(VAlign_Center)
 		[
 			SNew(SBox)
-			.WidthOverride(20)
-			.HeightOverride(20)
+			.WidthOverride(20.f)
+			.HeightOverride(20.f)
 			[
 				SNew(SImage)
 				.Image(Icon)
diff --git a/Engine/Source/Runtime/ApplicationCore/Public/Android/AndroidWindowUtils.h b/Engine/Source/Runtime/ApplicationCore/Public/Android/AndroidWindowUtils.h
--- a/Engine/Source/Runtime/ApplicationCore/Public/Android/AndroidWindowUtils.h
+++ b/Engine/Source/Runtime/ApplicationCore/Public/Android/AndroidWindowUtils.h
@@ -61,12 +61,12 @@ namespace AndroidWindowUtils
 				if (InOutScreenHeight > InOutScreenWidth)
 				{
 					Height = RequestedResX;
-					Width = (Height * ((float)InOutScreenWidth / (float)InOutScreenHeight) + 0.5f);
+					Width = FMath::TruncToInt32((float)Height * ((float)InOutScreenWidth / (float)InOutScreenHeight) + 0.5f);
 				}
 				else
 				{
 					Width = RequestedResX;
-					Height = (Width * ((float)InOutScreenHeight / (float)InOutScreenWidth) + 0.5f);
+					Height = FMath::TruncToInt32((float)Width * ((float)InOutScreenHeight / (float)InOutScreenWidth) + 0.5f);
 				}
 			}
 			else if (RequestedResY > 0)
@@ -75,12 +75,12 @@ namespace AndroidWindowUtils
 				if (InOutScreenHeight > InOutScreenWidth)
 				{
 					Width = RequestedResY;
-					Height = (Width * ((float)InOutScreenHeight / (float)InOutScreenWidth) + 0.5f);
+					Height = FMath::TruncToInt32((float)Width * ((float)InOutScreenHeight / (float)InOutScreenWidth) + 0.5f);
 				}
 				else
 				{
 					Height = RequestedResY;
-					Width = (Height * ((float)InOutScreenWidth / (float)InOutScreenHeight) + 0.5f);
+					Width = FMath::TruncToInt32((float)Height * ((float)InOutScreenWidth / (float)InOutScreenHeight) + 0.5f);
 				}
 			}
 			else
@@ -90,15 +90,15 @@ namespace AndroidWindowUtils
 				// set resolution relative to 1280x720 ((height > width == GAndroidIsPortrait)
 				if (InOutScreenHeight > InOutScreenWidth)
 				{
-					Height = 1280 * RequestedContentScaleFactor;
+					Height = FMath::TruncToInt32(1280.f * RequestedContentScaleFactor);
 				}
 				else
 				{
-					Height = 720 * RequestedContentScaleFactor;
+					Height = FMath::TruncToInt32(720.f * RequestedContentScaleFactor);
 				}
 
 				// apply the aspect ration to get the width
-				Width = (Height * AspectRatio + 0.5f);
+				Width = FMath::TruncToInt32((float)Height * AspectRatio + 0.5f);
 			}
 
 			// ensure Width and Height is multiple of 8
diff --git a/Engine/Source/Runtime/MediaAssets/Private/Assets/MediaSoundComponent.cpp b/Engine/Source/Runtime/MediaAssets/Private/Assets/MediaSoundComponent.cpp
--- a/Engine/Source/Runtime/MediaAssets/Private/Assets/MediaSoundComponent.cpp
+++ b/Engine/Source/Runtime/MediaAssets/Private/Assets/MediaSoundComponent.cpp
@@ -480,7 +480,7 @@ bool UMediaSoundComponent::Init(int32& SampleRate)
 ISoundGeneratorPtr UMediaSoundComponent::CreateSoundGenerator(const FSoundGeneratorInitParams& InParams)
 {
 	FMediaSoundGenerator::FSoundGeneratorParams Params;
-	Params.SampleRate = InParams.SampleRate;
+	Params.SampleRate = (int32)InParams.SampleRate;
 	Params.NumChannels = InParams.NumChannels;
 	Params.SampleQueue = SampleQueue;
 
diff --git a/Engine/Source/Runtime/Slate/Public/Framework/Application/SWindowTitleBar.h b/Engine/Source/Runtime/Slate/Public/Framework/Application/SWindowTitleBar.h
--- a/Engine/Source/Runtime/Slate/Public/Framework/Application/SWindowTitleBar.h
+++ b/Engine/Source/Runtime/Slate/Public/Framework/Application/SWindowTitleBar.h
@@ -450,7 +450,7 @@ protected:
 			RightSize = LeftSize;
 		}
 
-		float SpacerHeight = FMath::Max(LeftSize.Y, RightSize.Y);
+		float SpacerHeight = (float)FMath::Max(LeftSize.Y, RightSize.Y);
 
 		// create title bar
 		return 
diff --git a/Engine/Source/Runtime/Slate/Public/Widgets/Input/SEditableComboBox.h b/Engine/Source/Runtime/Slate/Public/Widgets/Input/SEditableComboBox.h
--- a/Engine/Source/Runtime/Slate/Public/Widgets/Input/SEditableComboBox.h
+++ b/Engine/Source/Runtime/Slate/Public/Widgets/Input/SEditableComboBox.h
@@ -146,7 +146,7 @@ public:
 								[
 									// rename button
 									SNew(SButton)
-										.ContentPadding(2)
+										.ContentPadding(2.f)
 										.ForegroundColor(FSlateColor::UseForeground())
 										.IsEnabled(this, &SEditableComboBox::HandleRemoveRenameButtonIsEnabled)
 										.ButtonStyle( FCoreStyle::Get(), "NoBorder" )
@@ -183,7 +183,7 @@ public:
 								[
 									// accept button
 									SNew(SButton)
-										.ContentPadding(2)
+										.ContentPadding(2.f)
 										.ForegroundColor(FSlateColor::UseForeground())
 										.ButtonStyle( FCoreStyle::Get(), "NoBorder" )
 										.OnClicked(this, &SEditableComboBox::HandleAcceptButtonClicked)
@@ -205,7 +205,7 @@ public:
 				[
 					// add button
 					SNew(SButton)
-						.ContentPadding(2)
+						.ContentPadding(2.f)
 						.ForegroundColor(FSlateColor::UseForeground())
 						.IsEnabled(this, &SEditableComboBox::HandleAddButtonIsEnabled)
 						.ButtonStyle( FCoreStyle::Get(), "NoBorder" )
@@ -226,7 +226,7 @@ public:
 				[
 					// remove button
 					SNew(SButton)
-						.ContentPadding(2)
+						.ContentPadding(2.f)
 						.ForegroundColor(FSlateColor::UseForeground())
 						.IsEnabled(this, &SEditableComboBox::HandleRemoveRenameButtonIsEnabled)
 						.ButtonStyle( FCoreStyle::Get(), "NoBorder" )
diff --git a/Engine/Source/Runtime/SlateCore/Public/Layout/Margin.h b/Engine/Source/Runtime/SlateCore/Public/Layout/Margin.h
--- a/Engine/Source/Runtime/SlateCore/Public/Layout/Margin.h
+++ b/Engine/Source/Runtime/SlateCore/Public/Layout/Margin.h
@@ -80,11 +80,11 @@ public:
 	{ }
 	
 	/** Construct a Margin where the spacing on each side is individually specified. */
-	FMargin(float InLeft, float InTop, float InRight, float InBottom)
-		: Left(InLeft)
-		, Top(InTop)
-		, Right(InRight)
-		, Bottom(InBottom)
+	FMargin(double InLeft, double InTop, double InRight, double InBottom)
+		: Left((float)InLeft)
+		, Top((float)InTop)
+		, Right((float)InRight)
+		, Bottom((float)InBottom)
 	{ }
 
 	/** Construct a Margin where the margins are coming from a FVector4 */
