Update Nintendo platform settings to use a Shading Mode enum ('Use Project Setting', 'Deferred', or 'Forward') rather than a "Use Forward Shading" checkbox in their platform project settings.
This brings it up to date with other platforms.

diff --git a/Engine/Source/Editor/PIEPreviewDeviceProfileSelector/Private/PIEPreviewDevice.cpp b/Engine/Source/Editor/PIEPreviewDeviceProfileSelector/Private/PIEPreviewDevice.cpp
--- a/Engine/Source/Editor/PIEPreviewDeviceProfileSelector/Private/PIEPreviewDevice.cpp
+++ b/Engine/Source/Editor/PIEPreviewDeviceProfileSelector/Private/PIEPreviewDevice.cpp
@@ -372,17 +372,31 @@ void FPIEPreviewDevice::ApplyRHIOverrides() const
 
 		case EPIEPreviewDeviceType::Switch:
 		{
-			// use a local ini file, so that Windows can read settings that are in SwitchEngine.ini files...
-			FConfigFile SwitchSettings;
-			FConfigCacheIni::LoadLocalIniFile(SwitchSettings, TEXT("Engine"), true, TEXT("Switch"));
-			bool bForwardShading = false;
-			SwitchSettings.GetBool(TEXT("/Script/SwitchRuntimeSettings.SwitchRuntimeSettings"), TEXT("bUseForwardShading"), bForwardShading);
-
 			// apply r.ForwardShading based on Switch bUseForwardShading setting
 			IConsoleVariable* CVarForwardShading = IConsoleManager::Get().FindConsoleVariable(TEXT("r.ForwardShading"));
 			if (CVarForwardShading)
 			{
-				CVarForwardShading->Set(bForwardShading ? 1 : 0);
+				// use a local ini file, so that Windows can read settings that are in SwitchEngine.ini files...
+				FConfigFile SwitchSettings;
+				FConfigCacheIni::LoadLocalIniFile(SwitchSettings, TEXT("Engine"), true, TEXT("Switch"));
+
+				FString ShadingModeString; // "UseProjectSetting", "Deferred", or "Forward"
+				if (GConfig->GetString(FPlatformProperties::GetRuntimeSettingsClassName(), TEXT("ShadingMode"), ShadingModeString, GEngineIni))
+				{
+					// Force forward shading
+					if (ShadingModeString.Equals(TEXT("Forward"), ESearchCase::IgnoreCase))
+					{
+						CVarForwardShading->Set(1);
+					}
+
+					// Force deferred shading
+					else if (ShadingModeString.Equals(TEXT("Deferred"), ESearchCase::IgnoreCase))
+					{
+						CVarForwardShading->Set(0);
+					}
+				}
+
+				// No need to support deprecated configurations here because they will have been automatically updated upon opening the editor
 			}
 		}
 		break;
