Mark editor modules as editor-only.

#Synced-CL 48550887

diff --git a/Engine/Source/Editor/AIGraph/AIGraph.Build.cs b/Engine/Source/Editor/AIGraph/AIGraph.Build.cs
--- a/Engine/Source/Editor/AIGraph/AIGraph.Build.cs
+++ b/Engine/Source/Editor/AIGraph/AIGraph.Build.cs
@@ -2,6 +2,7 @@
 
 using UnrealBuildTool;
 
+[SupportedTargetTypes(TargetType.Editor)]
 public class AIGraph : ModuleRules
 {
 	public AIGraph(ReadOnlyTargetRules Target) : base(Target)
diff --git a/Engine/Source/Editor/ActionableMessage/ActionableMessage.build.cs b/Engine/Source/Editor/ActionableMessage/ActionableMessage.build.cs
--- a/Engine/Source/Editor/ActionableMessage/ActionableMessage.build.cs
+++ b/Engine/Source/Editor/ActionableMessage/ActionableMessage.build.cs
@@ -2,6 +2,7 @@
 
 using UnrealBuildTool;
 
+[SupportedTargetTypes(TargetType.Editor)]
 public class ActionableMessage : ModuleRules
 {
 	public ActionableMessage(ReadOnlyTargetRules Target) : base(Target)
diff --git a/Engine/Source/Editor/ActorPickerMode/ActorPickerMode.Build.cs b/Engine/Source/Editor/ActorPickerMode/ActorPickerMode.Build.cs
--- a/Engine/Source/Editor/ActorPickerMode/ActorPickerMode.Build.cs
+++ b/Engine/Source/Editor/ActorPickerMode/ActorPickerMode.Build.cs
@@ -2,6 +2,7 @@
 
 using UnrealBuildTool;
 
+[SupportedTargetTypes(TargetType.Editor)]
 public class ActorPickerMode : ModuleRules
 {
     public ActorPickerMode(ReadOnlyTargetRules Target) : base(Target)
diff --git a/Engine/Source/Editor/AddContentDialog/AddContentDialog.Build.cs b/Engine/Source/Editor/AddContentDialog/AddContentDialog.Build.cs
--- a/Engine/Source/Editor/AddContentDialog/AddContentDialog.Build.cs
+++ b/Engine/Source/Editor/AddContentDialog/AddContentDialog.Build.cs
@@ -2,6 +2,7 @@
 
 using UnrealBuildTool;
 
+[SupportedTargetTypes(TargetType.Editor)]
 public class AddContentDialog : ModuleRules
 {
 	public AddContentDialog(ReadOnlyTargetRules Target) : base(Target)
diff --git a/Engine/Source/Editor/AdvancedPreviewScene/AdvancedPreviewScene.Build.cs b/Engine/Source/Editor/AdvancedPreviewScene/AdvancedPreviewScene.Build.cs
--- a/Engine/Source/Editor/AdvancedPreviewScene/AdvancedPreviewScene.Build.cs
+++ b/Engine/Source/Editor/AdvancedPreviewScene/AdvancedPreviewScene.Build.cs
@@ -2,6 +2,7 @@
 
 using UnrealBuildTool;
 
+[SupportedTargetTypes(TargetType.Editor)]
 public class AdvancedPreviewScene : ModuleRules
 {
     public AdvancedPreviewScene(ReadOnlyTargetRules Target) : base(Target)
diff --git a/Engine/Source/Editor/AnimGraph/AnimGraph.Build.cs b/Engine/Source/Editor/AnimGraph/AnimGraph.Build.cs
--- a/Engine/Source/Editor/AnimGraph/AnimGraph.Build.cs
+++ b/Engine/Source/Editor/AnimGraph/AnimGraph.Build.cs
@@ -3,6 +3,7 @@
 using UnrealBuildTool;
 using System.IO;
 
+[SupportedTargetTypes(TargetType.Editor)]
 public class AnimGraph : ModuleRules
 {
 	public AnimGraph(ReadOnlyTargetRules Target) : base(Target)
diff --git a/Engine/Source/Editor/AnimationBlueprintEditor/AnimationBlueprintEditor.Build.cs b/Engine/Source/Editor/AnimationBlueprintEditor/AnimationBlueprintEditor.Build.cs
--- a/Engine/Source/Editor/AnimationBlueprintEditor/AnimationBlueprintEditor.Build.cs
+++ b/Engine/Source/Editor/AnimationBlueprintEditor/AnimationBlueprintEditor.Build.cs
@@ -2,6 +2,7 @@
 
 using UnrealBuildTool;
 
+[SupportedTargetTypes(TargetType.Editor)]
 public class AnimationBlueprintEditor : ModuleRules
 {
 	public AnimationBlueprintEditor(ReadOnlyTargetRules Target) : base(Target)
diff --git a/Engine/Source/Editor/AnimationBlueprintLibrary/AnimationBlueprintLibrary.build.cs b/Engine/Source/Editor/AnimationBlueprintLibrary/AnimationBlueprintLibrary.build.cs
--- a/Engine/Source/Editor/AnimationBlueprintLibrary/AnimationBlueprintLibrary.build.cs
+++ b/Engine/Source/Editor/AnimationBlueprintLibrary/AnimationBlueprintLibrary.build.cs
@@ -2,6 +2,7 @@
 
 using UnrealBuildTool;
 
+[SupportedTargetTypes(TargetType.Editor)]
 public class AnimationBlueprintLibrary : ModuleRules
 {
 	public AnimationBlueprintLibrary(ReadOnlyTargetRules Target) : base(Target)
diff --git a/Engine/Source/Editor/AnimationEditMode/AnimationEditMode.Build.cs b/Engine/Source/Editor/AnimationEditMode/AnimationEditMode.Build.cs
--- a/Engine/Source/Editor/AnimationEditMode/AnimationEditMode.Build.cs
+++ b/Engine/Source/Editor/AnimationEditMode/AnimationEditMode.Build.cs
@@ -2,6 +2,7 @@
 
 using UnrealBuildTool;
 
+[SupportedTargetTypes(TargetType.Editor)]
 public class AnimationEditMode : ModuleRules
 {
 	public AnimationEditMode(ReadOnlyTargetRules Target) : base(Target)
diff --git a/Engine/Source/Editor/AnimationEditor/AnimationEditor.Build.cs b/Engine/Source/Editor/AnimationEditor/AnimationEditor.Build.cs
--- a/Engine/Source/Editor/AnimationEditor/AnimationEditor.Build.cs
+++ b/Engine/Source/Editor/AnimationEditor/AnimationEditor.Build.cs
@@ -2,6 +2,7 @@
 
 using UnrealBuildTool;
 
+[SupportedTargetTypes(TargetType.Editor)]
 public class AnimationEditor : ModuleRules
 {
 	public AnimationEditor(ReadOnlyTargetRules Target) : base(Target)
diff --git a/Engine/Source/Editor/AnimationEditorWidgets/AnimationEditorWidgets.Build.cs b/Engine/Source/Editor/AnimationEditorWidgets/AnimationEditorWidgets.Build.cs
--- a/Engine/Source/Editor/AnimationEditorWidgets/AnimationEditorWidgets.Build.cs
+++ b/Engine/Source/Editor/AnimationEditorWidgets/AnimationEditorWidgets.Build.cs
@@ -2,6 +2,7 @@
 
 using UnrealBuildTool;
 
+[SupportedTargetTypes(TargetType.Editor)]
 public class AnimationEditorWidgets : ModuleRules
 {
 	public AnimationEditorWidgets(ReadOnlyTargetRules Target) : base(Target)
diff --git a/Engine/Source/Editor/AnimationModifiers/AnimationModifiers.Build.cs b/Engine/Source/Editor/AnimationModifiers/AnimationModifiers.Build.cs
--- a/Engine/Source/Editor/AnimationModifiers/AnimationModifiers.Build.cs
+++ b/Engine/Source/Editor/AnimationModifiers/AnimationModifiers.Build.cs
@@ -2,6 +2,7 @@
 
 using UnrealBuildTool;
 
+[SupportedTargetTypes(TargetType.Editor)]
 public class AnimationModifiers : ModuleRules
 {
 	public AnimationModifiers(ReadOnlyTargetRules Target) : base(Target)
diff --git a/Engine/Source/Editor/AssetDefinition/AssetDefinition.Build.cs b/Engine/Source/Editor/AssetDefinition/AssetDefinition.Build.cs
--- a/Engine/Source/Editor/AssetDefinition/AssetDefinition.Build.cs
+++ b/Engine/Source/Editor/AssetDefinition/AssetDefinition.Build.cs
@@ -3,6 +3,7 @@
 using UnrealBuildTool;
 using System.IO;
 
+[SupportedTargetTypes(TargetType.Editor)]
 public class AssetDefinition : ModuleRules
 {
 	public AssetDefinition(ReadOnlyTargetRules Target) : base(Target)
diff --git a/Engine/Source/Editor/AssetTagsEditor/AssetTagsEditor.Build.cs b/Engine/Source/Editor/AssetTagsEditor/AssetTagsEditor.Build.cs
--- a/Engine/Source/Editor/AssetTagsEditor/AssetTagsEditor.Build.cs
+++ b/Engine/Source/Editor/AssetTagsEditor/AssetTagsEditor.Build.cs
@@ -2,6 +2,7 @@
 
 using UnrealBuildTool;
 
+[SupportedTargetTypes(TargetType.Editor)]
 public class AssetTagsEditor : ModuleRules
 {
 	public AssetTagsEditor(ReadOnlyTargetRules Target) : base(Target)
diff --git a/Engine/Source/Editor/AudioEditor/AudioEditor.Build.cs b/Engine/Source/Editor/AudioEditor/AudioEditor.Build.cs
--- a/Engine/Source/Editor/AudioEditor/AudioEditor.Build.cs
+++ b/Engine/Source/Editor/AudioEditor/AudioEditor.Build.cs
@@ -3,6 +3,7 @@
 using System.IO;
 using UnrealBuildTool;
 
+[SupportedTargetTypes(TargetType.Editor)]
 public class AudioEditor : ModuleRules
 {
 	public AudioEditor(ReadOnlyTargetRules Target) : base(Target)
diff --git a/Engine/Source/Editor/BehaviorTreeEditor/BehaviorTreeEditor.Build.cs b/Engine/Source/Editor/BehaviorTreeEditor/BehaviorTreeEditor.Build.cs
--- a/Engine/Source/Editor/BehaviorTreeEditor/BehaviorTreeEditor.Build.cs
+++ b/Engine/Source/Editor/BehaviorTreeEditor/BehaviorTreeEditor.Build.cs
@@ -2,6 +2,7 @@
 
 using UnrealBuildTool;
 
+[SupportedTargetTypes(TargetType.Editor)]
 public class BehaviorTreeEditor : ModuleRules
 {
 	public BehaviorTreeEditor(ReadOnlyTargetRules Target) : base(Target)
diff --git a/Engine/Source/Editor/BlueprintEditorLibrary/BlueprintEditorLibrary.Build.cs b/Engine/Source/Editor/BlueprintEditorLibrary/BlueprintEditorLibrary.Build.cs
--- a/Engine/Source/Editor/BlueprintEditorLibrary/BlueprintEditorLibrary.Build.cs
+++ b/Engine/Source/Editor/BlueprintEditorLibrary/BlueprintEditorLibrary.Build.cs
@@ -3,6 +3,7 @@
 using UnrealBuildTool;
 using System.IO;
 
+[SupportedTargetTypes(TargetType.Editor)]
 public class BlueprintEditorLibrary : ModuleRules
 {
 	public BlueprintEditorLibrary(ReadOnlyTargetRules Target) : base(Target)
diff --git a/Engine/Source/Editor/BlueprintGraph/BlueprintGraph.Build.cs b/Engine/Source/Editor/BlueprintGraph/BlueprintGraph.Build.cs
--- a/Engine/Source/Editor/BlueprintGraph/BlueprintGraph.Build.cs
+++ b/Engine/Source/Editor/BlueprintGraph/BlueprintGraph.Build.cs
@@ -3,6 +3,7 @@
 using UnrealBuildTool;
 using System.IO;
 
+[SupportedTargetTypes(TargetType.Editor)]
 public class BlueprintGraph : ModuleRules
 {
 	public BlueprintGraph(ReadOnlyTargetRules Target) : base(Target)
diff --git a/Engine/Source/Editor/Blutility/Blutility.Build.cs b/Engine/Source/Editor/Blutility/Blutility.Build.cs
--- a/Engine/Source/Editor/Blutility/Blutility.Build.cs
+++ b/Engine/Source/Editor/Blutility/Blutility.Build.cs
@@ -2,6 +2,7 @@
 
 using UnrealBuildTool;
 
+[SupportedTargetTypes(TargetType.Editor)]
 public class Blutility : ModuleRules
 {
 	public Blutility(ReadOnlyTargetRules Target) : base(Target)
diff --git a/Engine/Source/Editor/CSVtoSVG/CSVtoSVG.Build.cs b/Engine/Source/Editor/CSVtoSVG/CSVtoSVG.Build.cs
--- a/Engine/Source/Editor/CSVtoSVG/CSVtoSVG.Build.cs
+++ b/Engine/Source/Editor/CSVtoSVG/CSVtoSVG.Build.cs
@@ -2,6 +2,7 @@
 
 using UnrealBuildTool;
 
+[SupportedTargetTypes(TargetType.Editor)]
 public class CSVtoSVG : ModuleRules
 {
 	public CSVtoSVG(ReadOnlyTargetRules Target) : base(Target)
diff --git a/Engine/Source/Editor/ClassViewer/ClassViewer.Build.cs b/Engine/Source/Editor/ClassViewer/ClassViewer.Build.cs
--- a/Engine/Source/Editor/ClassViewer/ClassViewer.Build.cs
+++ b/Engine/Source/Editor/ClassViewer/ClassViewer.Build.cs
@@ -2,6 +2,7 @@
 
 using UnrealBuildTool;
 
+[SupportedTargetTypes(TargetType.Editor)]
 public class ClassViewer : ModuleRules
 {
 	public ClassViewer(ReadOnlyTargetRules Target) : base(Target)
diff --git a/Engine/Source/Editor/ClothPainter/ClothPainter.Build.cs b/Engine/Source/Editor/ClothPainter/ClothPainter.Build.cs
--- a/Engine/Source/Editor/ClothPainter/ClothPainter.Build.cs
+++ b/Engine/Source/Editor/ClothPainter/ClothPainter.Build.cs
@@ -2,6 +2,7 @@
 
 using UnrealBuildTool;
 
+[SupportedTargetTypes(TargetType.Editor)]
 public class ClothPainter : ModuleRules
 {
 	public ClothPainter(ReadOnlyTargetRules Target) : base(Target)
diff --git a/Engine/Source/Editor/ClothingSystemEditor/ClothingSystemEditor.Build.cs b/Engine/Source/Editor/ClothingSystemEditor/ClothingSystemEditor.Build.cs
--- a/Engine/Source/Editor/ClothingSystemEditor/ClothingSystemEditor.Build.cs
+++ b/Engine/Source/Editor/ClothingSystemEditor/ClothingSystemEditor.Build.cs
@@ -3,6 +3,7 @@
 using UnrealBuildTool;
 using System.IO;
 
+[SupportedTargetTypes(TargetType.Editor)]
 public class ClothingSystemEditor : ModuleRules
 {
 	public ClothingSystemEditor(ReadOnlyTargetRules Target) : base(Target)
diff --git a/Engine/Source/Editor/ClothingSystemEditorInterface/ClothingSystemEditorInterface.Build.cs b/Engine/Source/Editor/ClothingSystemEditorInterface/ClothingSystemEditorInterface.Build.cs
--- a/Engine/Source/Editor/ClothingSystemEditorInterface/ClothingSystemEditorInterface.Build.cs
+++ b/Engine/Source/Editor/ClothingSystemEditorInterface/ClothingSystemEditorInterface.Build.cs
@@ -3,6 +3,7 @@
 using UnrealBuildTool;
 using System.IO;
 
+[SupportedTargetTypes(TargetType.Editor)]
 public class ClothingSystemEditorInterface : ModuleRules
 {
 	public ClothingSystemEditorInterface(ReadOnlyTargetRules Target) : base(Target)
diff --git a/Engine/Source/Editor/CommonMenuExtensions/CommonMenuExtensions.Build.cs b/Engine/Source/Editor/CommonMenuExtensions/CommonMenuExtensions.Build.cs
--- a/Engine/Source/Editor/CommonMenuExtensions/CommonMenuExtensions.Build.cs
+++ b/Engine/Source/Editor/CommonMenuExtensions/CommonMenuExtensions.Build.cs
@@ -2,6 +2,7 @@
 
 using UnrealBuildTool;
 
+[SupportedTargetTypes(TargetType.Editor)]
 public class CommonMenuExtensions : ModuleRules
 {
 	// TODO: Is this a minimal enough list?
diff --git a/Engine/Source/Editor/ComponentVisualizers/ComponentVisualizers.Build.cs b/Engine/Source/Editor/ComponentVisualizers/ComponentVisualizers.Build.cs
--- a/Engine/Source/Editor/ComponentVisualizers/ComponentVisualizers.Build.cs
+++ b/Engine/Source/Editor/ComponentVisualizers/ComponentVisualizers.Build.cs
@@ -2,6 +2,7 @@
 
 using UnrealBuildTool;
 
+[SupportedTargetTypes(TargetType.Editor)]
 public class ComponentVisualizers : ModuleRules
 {
 	public ComponentVisualizers(ReadOnlyTargetRules Target) : base(Target)
diff --git a/Engine/Source/Editor/ConfigEditor/ConfigEditor.Build.cs b/Engine/Source/Editor/ConfigEditor/ConfigEditor.Build.cs
--- a/Engine/Source/Editor/ConfigEditor/ConfigEditor.Build.cs
+++ b/Engine/Source/Editor/ConfigEditor/ConfigEditor.Build.cs
@@ -2,6 +2,7 @@
 
 using UnrealBuildTool;
 
+[SupportedTargetTypes(TargetType.Editor)]
 public class ConfigEditor : ModuleRules
 {
 	public ConfigEditor(ReadOnlyTargetRules Target) : base(Target)
diff --git a/Engine/Source/Editor/ContentBrowser/ContentBrowser.Build.cs b/Engine/Source/Editor/ContentBrowser/ContentBrowser.Build.cs
--- a/Engine/Source/Editor/ContentBrowser/ContentBrowser.Build.cs
+++ b/Engine/Source/Editor/ContentBrowser/ContentBrowser.Build.cs
@@ -2,6 +2,7 @@
 
 using UnrealBuildTool;
 
+[SupportedTargetTypes(TargetType.Editor)]
 public class ContentBrowser : ModuleRules
 {
 	public ContentBrowser(ReadOnlyTargetRules Target) : base(Target)
diff --git a/Engine/Source/Editor/ContentBrowserData/ContentBrowserData.Build.cs b/Engine/Source/Editor/ContentBrowserData/ContentBrowserData.Build.cs
--- a/Engine/Source/Editor/ContentBrowserData/ContentBrowserData.Build.cs
+++ b/Engine/Source/Editor/ContentBrowserData/ContentBrowserData.Build.cs
@@ -2,6 +2,7 @@
 
 using UnrealBuildTool;
 
+[SupportedTargetTypes(TargetType.Editor)]
 public class ContentBrowserData : ModuleRules
 {
 	public ContentBrowserData(ReadOnlyTargetRules Target) : base(Target)
diff --git a/Engine/Source/Editor/CurveAssetEditor/CurveAssetEditor.Build.cs b/Engine/Source/Editor/CurveAssetEditor/CurveAssetEditor.Build.cs
--- a/Engine/Source/Editor/CurveAssetEditor/CurveAssetEditor.Build.cs
+++ b/Engine/Source/Editor/CurveAssetEditor/CurveAssetEditor.Build.cs
@@ -2,6 +2,7 @@
 
 using UnrealBuildTool;
 
+[SupportedTargetTypes(TargetType.Editor)]
 public class CurveAssetEditor : ModuleRules
 {
 	public CurveAssetEditor(ReadOnlyTargetRules Target) : base(Target)
diff --git a/Engine/Source/Editor/CurveEditor/CurveEditor.Build.cs b/Engine/Source/Editor/CurveEditor/CurveEditor.Build.cs
--- a/Engine/Source/Editor/CurveEditor/CurveEditor.Build.cs
+++ b/Engine/Source/Editor/CurveEditor/CurveEditor.Build.cs
@@ -2,6 +2,7 @@
 
 using UnrealBuildTool;
 
+[SupportedTargetTypes(TargetType.Editor)]
 public class CurveEditor : ModuleRules
 {
 	public CurveEditor(ReadOnlyTargetRules Target) : base(Target)
diff --git a/Engine/Source/Editor/CurveTableEditor/CurveTableEditor.Build.cs b/Engine/Source/Editor/CurveTableEditor/CurveTableEditor.Build.cs
--- a/Engine/Source/Editor/CurveTableEditor/CurveTableEditor.Build.cs
+++ b/Engine/Source/Editor/CurveTableEditor/CurveTableEditor.Build.cs
@@ -2,6 +2,7 @@
 
 using UnrealBuildTool;
 
+[SupportedTargetTypes(TargetType.Editor)]
 public class CurveTableEditor : ModuleRules
 {
 	public CurveTableEditor(ReadOnlyTargetRules Target) : base(Target)
diff --git a/Engine/Source/Editor/DataHierarchyEditor/DataHierarchyEditor.Build.cs b/Engine/Source/Editor/DataHierarchyEditor/DataHierarchyEditor.Build.cs
--- a/Engine/Source/Editor/DataHierarchyEditor/DataHierarchyEditor.Build.cs
+++ b/Engine/Source/Editor/DataHierarchyEditor/DataHierarchyEditor.Build.cs
@@ -2,6 +2,7 @@
 
 using UnrealBuildTool;
 
+[SupportedTargetTypes(TargetType.Editor)]
 public class DataHierarchyEditor : ModuleRules
 {
     public DataHierarchyEditor(ReadOnlyTargetRules Target) : base(Target)
diff --git a/Engine/Source/Editor/DataLayerEditor/DataLayerEditor.Build.cs b/Engine/Source/Editor/DataLayerEditor/DataLayerEditor.Build.cs
--- a/Engine/Source/Editor/DataLayerEditor/DataLayerEditor.Build.cs
+++ b/Engine/Source/Editor/DataLayerEditor/DataLayerEditor.Build.cs
@@ -2,6 +2,7 @@
 
 using UnrealBuildTool;
 
+[SupportedTargetTypes(TargetType.Editor)]
 public class DataLayerEditor : ModuleRules
 {
 	public DataLayerEditor(ReadOnlyTargetRules Target) : base(Target)
diff --git a/Engine/Source/Editor/DataTableEditor/DataTableEditor.Build.cs b/Engine/Source/Editor/DataTableEditor/DataTableEditor.Build.cs
--- a/Engine/Source/Editor/DataTableEditor/DataTableEditor.Build.cs
+++ b/Engine/Source/Editor/DataTableEditor/DataTableEditor.Build.cs
@@ -2,6 +2,7 @@
 
 using UnrealBuildTool;
 
+[SupportedTargetTypes(TargetType.Editor)]
 public class DataTableEditor : ModuleRules
 {
 	public DataTableEditor(ReadOnlyTargetRules Target) : base(Target)
diff --git a/Engine/Source/Editor/DerivedDataEditor/DerivedDataEditor.Build.cs b/Engine/Source/Editor/DerivedDataEditor/DerivedDataEditor.Build.cs
--- a/Engine/Source/Editor/DerivedDataEditor/DerivedDataEditor.Build.cs
+++ b/Engine/Source/Editor/DerivedDataEditor/DerivedDataEditor.Build.cs
@@ -1,6 +1,7 @@
 // Copyright Epic Games, Inc. All Rights Reserved.
 
 using UnrealBuildTool;
+[SupportedTargetTypes(TargetType.Editor)]
 public class DerivedDataEditor : ModuleRules
 {
 	public DerivedDataEditor(ReadOnlyTargetRules Target)
diff --git a/Engine/Source/Editor/DerivedDataWidgets/DerivedDataWidgets.Build.cs b/Engine/Source/Editor/DerivedDataWidgets/DerivedDataWidgets.Build.cs
--- a/Engine/Source/Editor/DerivedDataWidgets/DerivedDataWidgets.Build.cs
+++ b/Engine/Source/Editor/DerivedDataWidgets/DerivedDataWidgets.Build.cs
@@ -2,6 +2,7 @@
 
 using UnrealBuildTool;
 
+[SupportedTargetTypes(TargetType.Editor)]
 public class DerivedDataWidgets : ModuleRules
 {
 	public DerivedDataWidgets(ReadOnlyTargetRules Target)
diff --git a/Engine/Source/Editor/DetailCustomizations/DetailCustomizations.Build.cs b/Engine/Source/Editor/DetailCustomizations/DetailCustomizations.Build.cs
--- a/Engine/Source/Editor/DetailCustomizations/DetailCustomizations.Build.cs
+++ b/Engine/Source/Editor/DetailCustomizations/DetailCustomizations.Build.cs
@@ -2,6 +2,7 @@
 
 using UnrealBuildTool;
 
+[SupportedTargetTypes(TargetType.Editor)]
 public class DetailCustomizations : ModuleRules
 {
 	public DetailCustomizations(ReadOnlyTargetRules Target) : base(Target)
diff --git a/Engine/Source/Editor/DeviceProfileEditor/DeviceProfileEditor.Build.cs b/Engine/Source/Editor/DeviceProfileEditor/DeviceProfileEditor.Build.cs
--- a/Engine/Source/Editor/DeviceProfileEditor/DeviceProfileEditor.Build.cs
+++ b/Engine/Source/Editor/DeviceProfileEditor/DeviceProfileEditor.Build.cs
@@ -2,6 +2,7 @@
 
 using UnrealBuildTool;
 
+[SupportedTargetTypes(TargetType.Editor)]
 public class DeviceProfileEditor : ModuleRules
 {
 	public DeviceProfileEditor(ReadOnlyTargetRules Target) : base(Target)
diff --git a/Engine/Source/Editor/DeviceProfileServices/DeviceProfileServices.Build.cs b/Engine/Source/Editor/DeviceProfileServices/DeviceProfileServices.Build.cs
--- a/Engine/Source/Editor/DeviceProfileServices/DeviceProfileServices.Build.cs
+++ b/Engine/Source/Editor/DeviceProfileServices/DeviceProfileServices.Build.cs
@@ -2,6 +2,7 @@
 
 using UnrealBuildTool;
 
+[SupportedTargetTypes(TargetType.Editor)]
 public class DeviceProfileServices : ModuleRules
 {
 
diff --git a/Engine/Source/Editor/DistCurveEditor/DistCurveEditor.Build.cs b/Engine/Source/Editor/DistCurveEditor/DistCurveEditor.Build.cs
--- a/Engine/Source/Editor/DistCurveEditor/DistCurveEditor.Build.cs
+++ b/Engine/Source/Editor/DistCurveEditor/DistCurveEditor.Build.cs
@@ -2,6 +2,7 @@
 
 using UnrealBuildTool;
 
+[SupportedTargetTypes(TargetType.Editor)]
 public class DistCurveEditor : ModuleRules
 {
 	public DistCurveEditor(ReadOnlyTargetRules Target) : base(Target)
diff --git a/Engine/Source/Editor/Documentation/Documentation.Build.cs b/Engine/Source/Editor/Documentation/Documentation.Build.cs
--- a/Engine/Source/Editor/Documentation/Documentation.Build.cs
+++ b/Engine/Source/Editor/Documentation/Documentation.Build.cs
@@ -2,6 +2,7 @@
 
 namespace UnrealBuildTool.Rules
 {
+	[SupportedTargetTypes(TargetType.Editor)]
 	public class Documentation : ModuleRules
 	{
 		public Documentation(ReadOnlyTargetRules Target) : base(Target)
diff --git a/Engine/Source/Editor/EditorConfig/EditorConfig.Build.cs b/Engine/Source/Editor/EditorConfig/EditorConfig.Build.cs
--- a/Engine/Source/Editor/EditorConfig/EditorConfig.Build.cs
+++ b/Engine/Source/Editor/EditorConfig/EditorConfig.Build.cs
@@ -2,6 +2,7 @@
 
 namespace UnrealBuildTool.Rules
 {
+	[SupportedTargetTypes(TargetType.Editor)]
 	public class EditorConfig : ModuleRules
 	{
 		public EditorConfig(ReadOnlyTargetRules Target) : base(Target)
diff --git a/Engine/Source/Editor/EditorFramework/EditorFramework.Build.cs b/Engine/Source/Editor/EditorFramework/EditorFramework.Build.cs
--- a/Engine/Source/Editor/EditorFramework/EditorFramework.Build.cs
+++ b/Engine/Source/Editor/EditorFramework/EditorFramework.Build.cs
@@ -2,6 +2,7 @@
 
 using UnrealBuildTool;
 
+[SupportedTargetTypes(TargetType.Editor)]
 public class EditorFramework : ModuleRules
 {
 	public EditorFramework(ReadOnlyTargetRules Target) : base(Target)
diff --git a/Engine/Source/Editor/EditorSettingsViewer/EditorSettingsViewer.Build.cs b/Engine/Source/Editor/EditorSettingsViewer/EditorSettingsViewer.Build.cs
--- a/Engine/Source/Editor/EditorSettingsViewer/EditorSettingsViewer.Build.cs
+++ b/Engine/Source/Editor/EditorSettingsViewer/EditorSettingsViewer.Build.cs
@@ -2,6 +2,7 @@
 
 namespace UnrealBuildTool.Rules
 {
+	[SupportedTargetTypes(TargetType.Editor)]
 	public class EditorSettingsViewer : ModuleRules
 	{
 		public EditorSettingsViewer(ReadOnlyTargetRules Target) : base(Target)
diff --git a/Engine/Source/Editor/EditorStyle/EditorStyle.Build.cs b/Engine/Source/Editor/EditorStyle/EditorStyle.Build.cs
--- a/Engine/Source/Editor/EditorStyle/EditorStyle.Build.cs
+++ b/Engine/Source/Editor/EditorStyle/EditorStyle.Build.cs
@@ -2,6 +2,7 @@
 
 using UnrealBuildTool;
 
+[SupportedTargetTypes(TargetType.Editor)]
 public class EditorStyle : ModuleRules
 {
 	public EditorStyle(ReadOnlyTargetRules Target) : base(Target)
diff --git a/Engine/Source/Editor/EditorSubsystem/EditorSubsystem.Build.cs b/Engine/Source/Editor/EditorSubsystem/EditorSubsystem.Build.cs
--- a/Engine/Source/Editor/EditorSubsystem/EditorSubsystem.Build.cs
+++ b/Engine/Source/Editor/EditorSubsystem/EditorSubsystem.Build.cs
@@ -2,6 +2,7 @@
 
 using UnrealBuildTool;
 
+[SupportedTargetTypes(TargetType.Editor)]
 public class EditorSubsystem : ModuleRules
 {
     public EditorSubsystem(ReadOnlyTargetRules Target)
diff --git a/Engine/Source/Editor/EditorToolEvents/EditorToolEvents.Build.cs b/Engine/Source/Editor/EditorToolEvents/EditorToolEvents.Build.cs
--- a/Engine/Source/Editor/EditorToolEvents/EditorToolEvents.Build.cs
+++ b/Engine/Source/Editor/EditorToolEvents/EditorToolEvents.Build.cs
@@ -2,6 +2,7 @@
 
 using UnrealBuildTool;
 
+[SupportedTargetTypes(TargetType.Editor)]
 public class EditorToolEvents : ModuleRules
 {
 	public EditorToolEvents(ReadOnlyTargetRules Target) : base(Target)
diff --git a/Engine/Source/Editor/EditorWidgets/EditorWidgets.Build.cs b/Engine/Source/Editor/EditorWidgets/EditorWidgets.Build.cs
--- a/Engine/Source/Editor/EditorWidgets/EditorWidgets.Build.cs
+++ b/Engine/Source/Editor/EditorWidgets/EditorWidgets.Build.cs
@@ -2,6 +2,7 @@
 
 using UnrealBuildTool;
 
+[SupportedTargetTypes(TargetType.Editor)]
 public class EditorWidgets : ModuleRules
 {
 	public EditorWidgets(ReadOnlyTargetRules Target) : base(Target)
diff --git a/Engine/Source/Editor/EnvironmentLightingViewer/EnvironmentLightingViewer.Build.cs b/Engine/Source/Editor/EnvironmentLightingViewer/EnvironmentLightingViewer.Build.cs
--- a/Engine/Source/Editor/EnvironmentLightingViewer/EnvironmentLightingViewer.Build.cs
+++ b/Engine/Source/Editor/EnvironmentLightingViewer/EnvironmentLightingViewer.Build.cs
@@ -2,6 +2,7 @@
 
 using UnrealBuildTool;
 
+[SupportedTargetTypes(TargetType.Editor)]
 public class EnvironmentLightingViewer : ModuleRules
 {
 	public EnvironmentLightingViewer(ReadOnlyTargetRules Target) : base(Target)
diff --git a/Engine/Source/Editor/Experimental/EditorInteractiveToolsFramework/EditorInteractiveToolsFramework.Build.cs b/Engine/Source/Editor/Experimental/EditorInteractiveToolsFramework/EditorInteractiveToolsFramework.Build.cs
--- a/Engine/Source/Editor/Experimental/EditorInteractiveToolsFramework/EditorInteractiveToolsFramework.Build.cs
+++ b/Engine/Source/Editor/Experimental/EditorInteractiveToolsFramework/EditorInteractiveToolsFramework.Build.cs
@@ -2,6 +2,7 @@
 
 using UnrealBuildTool;
 
+[SupportedTargetTypes(TargetType.Editor)]
 public class EditorInteractiveToolsFramework : ModuleRules
 {
 	public EditorInteractiveToolsFramework(ReadOnlyTargetRules Target) : base(Target)
diff --git a/Engine/Source/Editor/FoliageEdit/FoliageEdit.build.cs b/Engine/Source/Editor/FoliageEdit/FoliageEdit.build.cs
--- a/Engine/Source/Editor/FoliageEdit/FoliageEdit.build.cs
+++ b/Engine/Source/Editor/FoliageEdit/FoliageEdit.build.cs
@@ -2,6 +2,7 @@
 
 using UnrealBuildTool;
 
+[SupportedTargetTypes(TargetType.Editor)]
 public class FoliageEdit : ModuleRules
 {
 	public FoliageEdit(ReadOnlyTargetRules Target) : base(Target)
diff --git a/Engine/Source/Editor/FontEditor/FontEditor.Build.cs b/Engine/Source/Editor/FontEditor/FontEditor.Build.cs
--- a/Engine/Source/Editor/FontEditor/FontEditor.Build.cs
+++ b/Engine/Source/Editor/FontEditor/FontEditor.Build.cs
@@ -2,6 +2,7 @@
 
 using UnrealBuildTool;
 
+[SupportedTargetTypes(TargetType.Editor)]
 public class FontEditor : ModuleRules
 {
 	public FontEditor(ReadOnlyTargetRules Target) : base(Target)
diff --git a/Engine/Source/Editor/GameProjectGeneration/GameProjectGeneration.Build.cs b/Engine/Source/Editor/GameProjectGeneration/GameProjectGeneration.Build.cs
--- a/Engine/Source/Editor/GameProjectGeneration/GameProjectGeneration.Build.cs
+++ b/Engine/Source/Editor/GameProjectGeneration/GameProjectGeneration.Build.cs
@@ -2,6 +2,7 @@
 
 using UnrealBuildTool;
 
+[SupportedTargetTypes(TargetType.Editor)]
 public class GameProjectGeneration : ModuleRules
 {
 	public GameProjectGeneration(ReadOnlyTargetRules Target) : base(Target)
diff --git a/Engine/Source/Editor/GameplayDebugger/GameplayDebuggerEditor.Build.cs b/Engine/Source/Editor/GameplayDebugger/GameplayDebuggerEditor.Build.cs
--- a/Engine/Source/Editor/GameplayDebugger/GameplayDebuggerEditor.Build.cs
+++ b/Engine/Source/Editor/GameplayDebugger/GameplayDebuggerEditor.Build.cs
@@ -2,6 +2,7 @@
 
 namespace UnrealBuildTool.Rules
 {
+	[SupportedTargetTypes(TargetType.Editor)]
 	public class GameplayDebuggerEditor : ModuleRules
 	{
 		public GameplayDebuggerEditor(ReadOnlyTargetRules Target) : base(Target)
diff --git a/Engine/Source/Editor/GameplayTasksEditor/GameplayTasksEditor.Build.cs b/Engine/Source/Editor/GameplayTasksEditor/GameplayTasksEditor.Build.cs
--- a/Engine/Source/Editor/GameplayTasksEditor/GameplayTasksEditor.Build.cs
+++ b/Engine/Source/Editor/GameplayTasksEditor/GameplayTasksEditor.Build.cs
@@ -2,6 +2,7 @@
 
 namespace UnrealBuildTool.Rules
 {
+	[SupportedTargetTypes(TargetType.Editor)]
 	public class GameplayTasksEditor : ModuleRules
 	{
         public GameplayTasksEditor(ReadOnlyTargetRules Target) : base(Target)
diff --git a/Engine/Source/Editor/GraphEditor/GraphEditor.Build.cs b/Engine/Source/Editor/GraphEditor/GraphEditor.Build.cs
--- a/Engine/Source/Editor/GraphEditor/GraphEditor.Build.cs
+++ b/Engine/Source/Editor/GraphEditor/GraphEditor.Build.cs
@@ -2,6 +2,7 @@
 
 using UnrealBuildTool;
 
+[SupportedTargetTypes(TargetType.Editor)]
 public class GraphEditor : ModuleRules
 {
 	public GraphEditor(ReadOnlyTargetRules Target) : base(Target)
diff --git a/Engine/Source/Editor/HardwareTargeting/HardwareTargeting.Build.cs b/Engine/Source/Editor/HardwareTargeting/HardwareTargeting.Build.cs
--- a/Engine/Source/Editor/HardwareTargeting/HardwareTargeting.Build.cs
+++ b/Engine/Source/Editor/HardwareTargeting/HardwareTargeting.Build.cs
@@ -2,6 +2,7 @@
 
 using UnrealBuildTool;
 
+[SupportedTargetTypes(TargetType.Editor)]
 public class HardwareTargeting : ModuleRules
 {
 	public HardwareTargeting(ReadOnlyTargetRules Target) : base(Target)
diff --git a/Engine/Source/Editor/HierarchicalLODOutliner/HierarchicalLODOutliner.Build.cs b/Engine/Source/Editor/HierarchicalLODOutliner/HierarchicalLODOutliner.Build.cs
--- a/Engine/Source/Editor/HierarchicalLODOutliner/HierarchicalLODOutliner.Build.cs
+++ b/Engine/Source/Editor/HierarchicalLODOutliner/HierarchicalLODOutliner.Build.cs
@@ -2,6 +2,7 @@
 
 using UnrealBuildTool;
 
+[SupportedTargetTypes(TargetType.Editor)]
 public class HierarchicalLODOutliner : ModuleRules
 {
     public HierarchicalLODOutliner(ReadOnlyTargetRules Target) : base(Target)
diff --git a/Engine/Source/Editor/InputBindingEditor/InputBindingEditor.Build.cs b/Engine/Source/Editor/InputBindingEditor/InputBindingEditor.Build.cs
--- a/Engine/Source/Editor/InputBindingEditor/InputBindingEditor.Build.cs
+++ b/Engine/Source/Editor/InputBindingEditor/InputBindingEditor.Build.cs
@@ -2,6 +2,7 @@
 
 using UnrealBuildTool;
 
+[SupportedTargetTypes(TargetType.Editor)]
 public class InputBindingEditor : ModuleRules
 {
 	public InputBindingEditor(ReadOnlyTargetRules Target) : base(Target)
diff --git a/Engine/Source/Editor/InternationalizationSettings/InternationalizationSettings.Build.cs b/Engine/Source/Editor/InternationalizationSettings/InternationalizationSettings.Build.cs
--- a/Engine/Source/Editor/InternationalizationSettings/InternationalizationSettings.Build.cs
+++ b/Engine/Source/Editor/InternationalizationSettings/InternationalizationSettings.Build.cs
@@ -2,6 +2,7 @@
 
 namespace UnrealBuildTool.Rules
 {
+	[SupportedTargetTypes(TargetType.Editor)]
 	public class InternationalizationSettings : ModuleRules
 	{
 		public InternationalizationSettings(ReadOnlyTargetRules Target) : base(Target)
diff --git a/Engine/Source/Editor/Kismet/Kismet.Build.cs b/Engine/Source/Editor/Kismet/Kismet.Build.cs
--- a/Engine/Source/Editor/Kismet/Kismet.Build.cs
+++ b/Engine/Source/Editor/Kismet/Kismet.Build.cs
@@ -2,6 +2,7 @@
 
 using UnrealBuildTool;
 
+[SupportedTargetTypes(TargetType.Editor)]
 public class Kismet : ModuleRules
 {
 	public Kismet(ReadOnlyTargetRules Target) : base(Target)
diff --git a/Engine/Source/Editor/KismetCompiler/KismetCompiler.Build.cs b/Engine/Source/Editor/KismetCompiler/KismetCompiler.Build.cs
--- a/Engine/Source/Editor/KismetCompiler/KismetCompiler.Build.cs
+++ b/Engine/Source/Editor/KismetCompiler/KismetCompiler.Build.cs
@@ -2,6 +2,7 @@
 
 using UnrealBuildTool;
 
+[SupportedTargetTypes(TargetType.Editor)]
 public class KismetCompiler : ModuleRules
 {
 	public KismetCompiler(ReadOnlyTargetRules Target) : base(Target)
diff --git a/Engine/Source/Editor/KismetWidgets/KismetWidgets.Build.cs b/Engine/Source/Editor/KismetWidgets/KismetWidgets.Build.cs
--- a/Engine/Source/Editor/KismetWidgets/KismetWidgets.Build.cs
+++ b/Engine/Source/Editor/KismetWidgets/KismetWidgets.Build.cs
@@ -2,6 +2,7 @@
 
 using UnrealBuildTool;
 
+[SupportedTargetTypes(TargetType.Editor)]
 public class KismetWidgets : ModuleRules
 {
 	public KismetWidgets(ReadOnlyTargetRules Target) : base(Target)
diff --git a/Engine/Source/Editor/LandscapeEditor/LandscapeEditor.Build.cs b/Engine/Source/Editor/LandscapeEditor/LandscapeEditor.Build.cs
--- a/Engine/Source/Editor/LandscapeEditor/LandscapeEditor.Build.cs
+++ b/Engine/Source/Editor/LandscapeEditor/LandscapeEditor.Build.cs
@@ -2,6 +2,7 @@
 
 using UnrealBuildTool;
 
+[SupportedTargetTypes(TargetType.Editor)]
 public class LandscapeEditor : ModuleRules
 {
 	public LandscapeEditor(ReadOnlyTargetRules Target) : base(Target)
diff --git a/Engine/Source/Editor/LandscapeEditorUtilities/LandscapeEditorUtilities.Build.cs b/Engine/Source/Editor/LandscapeEditorUtilities/LandscapeEditorUtilities.Build.cs
--- a/Engine/Source/Editor/LandscapeEditorUtilities/LandscapeEditorUtilities.Build.cs
+++ b/Engine/Source/Editor/LandscapeEditorUtilities/LandscapeEditorUtilities.Build.cs
@@ -2,6 +2,7 @@
 
 using UnrealBuildTool;
 
+[SupportedTargetTypes(TargetType.Editor)]
 public class LandscapeEditorUtilities : ModuleRules
 {
 	public LandscapeEditorUtilities(ReadOnlyTargetRules Target) : base(Target)
diff --git a/Engine/Source/Editor/Layers/Layers.Build.cs b/Engine/Source/Editor/Layers/Layers.Build.cs
--- a/Engine/Source/Editor/Layers/Layers.Build.cs
+++ b/Engine/Source/Editor/Layers/Layers.Build.cs
@@ -2,6 +2,7 @@
 
 using UnrealBuildTool;
 
+[SupportedTargetTypes(TargetType.Editor)]
 public class Layers : ModuleRules
 {
 	public Layers(ReadOnlyTargetRules Target) : base(Target)
diff --git a/Engine/Source/Editor/LevelEditor/LevelEditor.Build.cs b/Engine/Source/Editor/LevelEditor/LevelEditor.Build.cs
--- a/Engine/Source/Editor/LevelEditor/LevelEditor.Build.cs
+++ b/Engine/Source/Editor/LevelEditor/LevelEditor.Build.cs
@@ -2,6 +2,7 @@
 
 using UnrealBuildTool;
 
+[SupportedTargetTypes(TargetType.Editor)]
 public class LevelEditor : ModuleRules
 {
 	public LevelEditor(ReadOnlyTargetRules Target) : base(Target)
diff --git a/Engine/Source/Editor/LevelInstanceEditor/LevelInstanceEditor.Build.cs b/Engine/Source/Editor/LevelInstanceEditor/LevelInstanceEditor.Build.cs
--- a/Engine/Source/Editor/LevelInstanceEditor/LevelInstanceEditor.Build.cs
+++ b/Engine/Source/Editor/LevelInstanceEditor/LevelInstanceEditor.Build.cs
@@ -2,6 +2,7 @@
 
 using UnrealBuildTool;
 
+[SupportedTargetTypes(TargetType.Editor)]
 public class LevelInstanceEditor : ModuleRules
 {
     public LevelInstanceEditor(ReadOnlyTargetRules Target) : base(Target)
diff --git a/Engine/Source/Editor/LocalizationCommandletExecution/LocalizationCommandletExecution.Build.cs b/Engine/Source/Editor/LocalizationCommandletExecution/LocalizationCommandletExecution.Build.cs
--- a/Engine/Source/Editor/LocalizationCommandletExecution/LocalizationCommandletExecution.Build.cs
+++ b/Engine/Source/Editor/LocalizationCommandletExecution/LocalizationCommandletExecution.Build.cs
@@ -2,6 +2,7 @@
 
 using UnrealBuildTool;
 
+[SupportedTargetTypes(TargetType.Editor)]
 public class LocalizationCommandletExecution : ModuleRules
 {
 	public LocalizationCommandletExecution(ReadOnlyTargetRules Target) : base(Target)
diff --git a/Engine/Source/Editor/LocalizationDashboard/LocalizationDashboard.Build.cs b/Engine/Source/Editor/LocalizationDashboard/LocalizationDashboard.Build.cs
--- a/Engine/Source/Editor/LocalizationDashboard/LocalizationDashboard.Build.cs
+++ b/Engine/Source/Editor/LocalizationDashboard/LocalizationDashboard.Build.cs
@@ -2,6 +2,7 @@
 
 using UnrealBuildTool;
 
+[SupportedTargetTypes(TargetType.Editor)]
 public class LocalizationDashboard : ModuleRules
 {
 	public LocalizationDashboard(ReadOnlyTargetRules Target) : base(Target)
diff --git a/Engine/Source/Editor/MainFrame/MainFrame.Build.cs b/Engine/Source/Editor/MainFrame/MainFrame.Build.cs
--- a/Engine/Source/Editor/MainFrame/MainFrame.Build.cs
+++ b/Engine/Source/Editor/MainFrame/MainFrame.Build.cs
@@ -2,6 +2,7 @@
 
 using UnrealBuildTool;
 
+[SupportedTargetTypes(TargetType.Editor)]
 public class MainFrame : ModuleRules
 {
 	public MainFrame(ReadOnlyTargetRules Target) : base(Target)
diff --git a/Engine/Source/Editor/MassEntityDebugger/MassEntityDebugger.Build.cs b/Engine/Source/Editor/MassEntityDebugger/MassEntityDebugger.Build.cs
--- a/Engine/Source/Editor/MassEntityDebugger/MassEntityDebugger.Build.cs
+++ b/Engine/Source/Editor/MassEntityDebugger/MassEntityDebugger.Build.cs
@@ -2,6 +2,7 @@
 
 namespace UnrealBuildTool.Rules
 {
+	[SupportedTargetTypes(TargetType.Editor)]
 	public class MassEntityDebugger : ModuleRules
 	{
 		public MassEntityDebugger(ReadOnlyTargetRules Target) : base(Target)
diff --git a/Engine/Source/Editor/MassEntityEditor/MassEntityEditor.Build.cs b/Engine/Source/Editor/MassEntityEditor/MassEntityEditor.Build.cs
--- a/Engine/Source/Editor/MassEntityEditor/MassEntityEditor.Build.cs
+++ b/Engine/Source/Editor/MassEntityEditor/MassEntityEditor.Build.cs
@@ -2,6 +2,7 @@
 
 namespace UnrealBuildTool.Rules
 {
+	[SupportedTargetTypes(TargetType.Editor)]
 	public class MassEntityEditor : ModuleRules
 	{
 		public MassEntityEditor(ReadOnlyTargetRules Target) : base(Target)
diff --git a/Engine/Source/Editor/MaterialEditor/MaterialEditor.Build.cs b/Engine/Source/Editor/MaterialEditor/MaterialEditor.Build.cs
--- a/Engine/Source/Editor/MaterialEditor/MaterialEditor.Build.cs
+++ b/Engine/Source/Editor/MaterialEditor/MaterialEditor.Build.cs
@@ -2,6 +2,7 @@
 
 using UnrealBuildTool;
 
+[SupportedTargetTypes(TargetType.Editor)]
 public class MaterialEditor : ModuleRules
 {
 	public MaterialEditor(ReadOnlyTargetRules Target) : base(Target)
diff --git a/Engine/Source/Editor/MergeActors/MergeActors.Build.cs b/Engine/Source/Editor/MergeActors/MergeActors.Build.cs
--- a/Engine/Source/Editor/MergeActors/MergeActors.Build.cs
+++ b/Engine/Source/Editor/MergeActors/MergeActors.Build.cs
@@ -2,6 +2,7 @@
 
 using UnrealBuildTool;
 
+[SupportedTargetTypes(TargetType.Editor)]
 public class MergeActors : ModuleRules
 {
 	public MergeActors(ReadOnlyTargetRules Target) : base(Target)
diff --git a/Engine/Source/Editor/MeshPaint/MeshPaint.Build.cs b/Engine/Source/Editor/MeshPaint/MeshPaint.Build.cs
--- a/Engine/Source/Editor/MeshPaint/MeshPaint.Build.cs
+++ b/Engine/Source/Editor/MeshPaint/MeshPaint.Build.cs
@@ -2,6 +2,7 @@
 
 using UnrealBuildTool;
 
+[SupportedTargetTypes(TargetType.Editor)]
 public class MeshPaint : ModuleRules
 {
     public MeshPaint(ReadOnlyTargetRules Target) : base(Target)
diff --git a/Engine/Source/Editor/MovieSceneCaptureDialog/MovieSceneCaptureDialog.Build.cs b/Engine/Source/Editor/MovieSceneCaptureDialog/MovieSceneCaptureDialog.Build.cs
--- a/Engine/Source/Editor/MovieSceneCaptureDialog/MovieSceneCaptureDialog.Build.cs
+++ b/Engine/Source/Editor/MovieSceneCaptureDialog/MovieSceneCaptureDialog.Build.cs
@@ -2,6 +2,7 @@
 
 using UnrealBuildTool;
 
+[SupportedTargetTypes(TargetType.Editor)]
 public class MovieSceneCaptureDialog : ModuleRules
 {
 	public MovieSceneCaptureDialog(ReadOnlyTargetRules Target) : base(Target)
diff --git a/Engine/Source/Editor/MovieSceneTools/MovieSceneTools.Build.cs b/Engine/Source/Editor/MovieSceneTools/MovieSceneTools.Build.cs
--- a/Engine/Source/Editor/MovieSceneTools/MovieSceneTools.Build.cs
+++ b/Engine/Source/Editor/MovieSceneTools/MovieSceneTools.Build.cs
@@ -2,6 +2,7 @@
 
 using UnrealBuildTool;
 
+[SupportedTargetTypes(TargetType.Editor)]
 public class MovieSceneTools : ModuleRules
 {
 	public MovieSceneTools(ReadOnlyTargetRules Target) : base(Target)
diff --git a/Engine/Source/Editor/NNEEditor/NNEEditor.Build.cs b/Engine/Source/Editor/NNEEditor/NNEEditor.Build.cs
--- a/Engine/Source/Editor/NNEEditor/NNEEditor.Build.cs
+++ b/Engine/Source/Editor/NNEEditor/NNEEditor.Build.cs
@@ -3,6 +3,7 @@
 using System.IO;
 using UnrealBuildTool;
 
+[SupportedTargetTypes(TargetType.Editor)]
 public class NNEEditor : ModuleRules
 {
 	public NNEEditor(ReadOnlyTargetRules Target) : base(Target)
diff --git a/Engine/Source/Editor/NewLevelDialog/NewLevelDialog.Build.cs b/Engine/Source/Editor/NewLevelDialog/NewLevelDialog.Build.cs
--- a/Engine/Source/Editor/NewLevelDialog/NewLevelDialog.Build.cs
+++ b/Engine/Source/Editor/NewLevelDialog/NewLevelDialog.Build.cs
@@ -2,6 +2,7 @@
 
 using UnrealBuildTool;
 
+[SupportedTargetTypes(TargetType.Editor)]
 public class NewLevelDialog : ModuleRules
 {
 	public NewLevelDialog(ReadOnlyTargetRules Target) : base(Target)
diff --git a/Engine/Source/Editor/OverlayEditor/OverlayEditor.Build.cs b/Engine/Source/Editor/OverlayEditor/OverlayEditor.Build.cs
--- a/Engine/Source/Editor/OverlayEditor/OverlayEditor.Build.cs
+++ b/Engine/Source/Editor/OverlayEditor/OverlayEditor.Build.cs
@@ -2,6 +2,7 @@
 
 namespace UnrealBuildTool.Rules
 {
+	[SupportedTargetTypes(TargetType.Editor)]
 	public class OverlayEditor : ModuleRules
 	{
 		public OverlayEditor(ReadOnlyTargetRules Target) : base(Target)
diff --git a/Engine/Source/Editor/PIEPreviewDeviceProfileSelector/PIEPreviewDeviceProfileSelector.Build.cs b/Engine/Source/Editor/PIEPreviewDeviceProfileSelector/PIEPreviewDeviceProfileSelector.Build.cs
--- a/Engine/Source/Editor/PIEPreviewDeviceProfileSelector/PIEPreviewDeviceProfileSelector.Build.cs
+++ b/Engine/Source/Editor/PIEPreviewDeviceProfileSelector/PIEPreviewDeviceProfileSelector.Build.cs
@@ -2,7 +2,8 @@
 
 namespace UnrealBuildTool.Rules
 {
-    public class PIEPreviewDeviceProfileSelector : ModuleRules
+    [SupportedTargetTypes(TargetType.Editor)]
+	public class PIEPreviewDeviceProfileSelector : ModuleRules
 	{
         public PIEPreviewDeviceProfileSelector(ReadOnlyTargetRules Target) : base(Target)
 		{
diff --git a/Engine/Source/Editor/PIEPreviewDeviceSpecification/PIEPreviewDeviceSpecification.Build.cs b/Engine/Source/Editor/PIEPreviewDeviceSpecification/PIEPreviewDeviceSpecification.Build.cs
--- a/Engine/Source/Editor/PIEPreviewDeviceSpecification/PIEPreviewDeviceSpecification.Build.cs
+++ b/Engine/Source/Editor/PIEPreviewDeviceSpecification/PIEPreviewDeviceSpecification.Build.cs
@@ -2,7 +2,8 @@
 
 namespace UnrealBuildTool.Rules
 {
-    public class PIEPreviewDeviceSpecification : ModuleRules
+	[SupportedTargetTypes(TargetType.Editor, TargetType.Program)]
+	public class PIEPreviewDeviceSpecification : ModuleRules
 	{
         public PIEPreviewDeviceSpecification(ReadOnlyTargetRules Target) : base(Target)
 		{
diff --git a/Engine/Source/Editor/PackagesDialog/PackagesDialog.Build.cs b/Engine/Source/Editor/PackagesDialog/PackagesDialog.Build.cs
--- a/Engine/Source/Editor/PackagesDialog/PackagesDialog.Build.cs
+++ b/Engine/Source/Editor/PackagesDialog/PackagesDialog.Build.cs
@@ -2,6 +2,7 @@
 
 using UnrealBuildTool;
 
+[SupportedTargetTypes(TargetType.Editor)]
 public class PackagesDialog : ModuleRules
 {
 	public PackagesDialog(ReadOnlyTargetRules Target) : base(Target)
diff --git a/Engine/Source/Editor/Persona/Persona.Build.cs b/Engine/Source/Editor/Persona/Persona.Build.cs
--- a/Engine/Source/Editor/Persona/Persona.Build.cs
+++ b/Engine/Source/Editor/Persona/Persona.Build.cs
@@ -2,6 +2,7 @@
 
 using UnrealBuildTool;
 
+[SupportedTargetTypes(TargetType.Editor)]
 public class Persona : ModuleRules
 {
     public Persona(ReadOnlyTargetRules Target) : base(Target)
diff --git a/Engine/Source/Editor/PhysicsAssetEditor/PhysicsAssetEditor.Build.cs b/Engine/Source/Editor/PhysicsAssetEditor/PhysicsAssetEditor.Build.cs
--- a/Engine/Source/Editor/PhysicsAssetEditor/PhysicsAssetEditor.Build.cs
+++ b/Engine/Source/Editor/PhysicsAssetEditor/PhysicsAssetEditor.Build.cs
@@ -2,6 +2,7 @@
 
 using UnrealBuildTool;
 
+[SupportedTargetTypes(TargetType.Editor)]
 public class PhysicsAssetEditor : ModuleRules
 {
 	public PhysicsAssetEditor(ReadOnlyTargetRules Target) : base(Target)
diff --git a/Engine/Source/Editor/PinnedCommandList/PinnedCommandList.Build.cs b/Engine/Source/Editor/PinnedCommandList/PinnedCommandList.Build.cs
--- a/Engine/Source/Editor/PinnedCommandList/PinnedCommandList.Build.cs
+++ b/Engine/Source/Editor/PinnedCommandList/PinnedCommandList.Build.cs
@@ -2,6 +2,7 @@
 
 using UnrealBuildTool;
 
+[SupportedTargetTypes(TargetType.Editor)]
 public class PinnedCommandList : ModuleRules
 {
     public PinnedCommandList(ReadOnlyTargetRules Target) : base(Target)
diff --git a/Engine/Source/Editor/PixelInspector/PixelInspectorModule.Build.cs b/Engine/Source/Editor/PixelInspector/PixelInspectorModule.Build.cs
--- a/Engine/Source/Editor/PixelInspector/PixelInspectorModule.Build.cs
+++ b/Engine/Source/Editor/PixelInspector/PixelInspectorModule.Build.cs
@@ -2,6 +2,7 @@
 
 using UnrealBuildTool;
 
+[SupportedTargetTypes(TargetType.Editor)]
 public class PixelInspectorModule : ModuleRules
 {
     public PixelInspectorModule(ReadOnlyTargetRules Target) : base(Target)
diff --git a/Engine/Source/Editor/PlacementMode/PlacementMode.Build.cs b/Engine/Source/Editor/PlacementMode/PlacementMode.Build.cs
--- a/Engine/Source/Editor/PlacementMode/PlacementMode.Build.cs
+++ b/Engine/Source/Editor/PlacementMode/PlacementMode.Build.cs
@@ -2,6 +2,7 @@
 
 using UnrealBuildTool;
 
+[SupportedTargetTypes(TargetType.Editor)]
 public class PlacementMode : ModuleRules
 {
     public PlacementMode(ReadOnlyTargetRules Target) : base(Target)
diff --git a/Engine/Source/Editor/ProjectSettingsViewer/ProjectSettingsViewer.Build.cs b/Engine/Source/Editor/ProjectSettingsViewer/ProjectSettingsViewer.Build.cs
--- a/Engine/Source/Editor/ProjectSettingsViewer/ProjectSettingsViewer.Build.cs
+++ b/Engine/Source/Editor/ProjectSettingsViewer/ProjectSettingsViewer.Build.cs
@@ -2,6 +2,7 @@
 
 namespace UnrealBuildTool.Rules
 {
+	[SupportedTargetTypes(TargetType.Editor)]
 	public class ProjectSettingsViewer : ModuleRules
 	{
 		public ProjectSettingsViewer(ReadOnlyTargetRules Target) : base(Target)
diff --git a/Engine/Source/Editor/ProjectTargetPlatformEditor/ProjectTargetPlatformEditor.Build.cs b/Engine/Source/Editor/ProjectTargetPlatformEditor/ProjectTargetPlatformEditor.Build.cs
--- a/Engine/Source/Editor/ProjectTargetPlatformEditor/ProjectTargetPlatformEditor.Build.cs
+++ b/Engine/Source/Editor/ProjectTargetPlatformEditor/ProjectTargetPlatformEditor.Build.cs
@@ -2,6 +2,7 @@
 
 using UnrealBuildTool;
 
+[SupportedTargetTypes(TargetType.Editor)]
 public class ProjectTargetPlatformEditor : ModuleRules
 {
 	public ProjectTargetPlatformEditor(ReadOnlyTargetRules Target) : base(Target)
diff --git a/Engine/Source/Editor/PropertyEditor/PropertyEditor.Build.cs b/Engine/Source/Editor/PropertyEditor/PropertyEditor.Build.cs
--- a/Engine/Source/Editor/PropertyEditor/PropertyEditor.Build.cs
+++ b/Engine/Source/Editor/PropertyEditor/PropertyEditor.Build.cs
@@ -2,6 +2,7 @@
 
 using UnrealBuildTool;
 
+[SupportedTargetTypes(TargetType.Editor)]
 public class PropertyEditor : ModuleRules
 {
 	public PropertyEditor(ReadOnlyTargetRules Target) : base(Target)
diff --git a/Engine/Source/Editor/RenderResourceViewer/RenderResourceViewer.Build.cs b/Engine/Source/Editor/RenderResourceViewer/RenderResourceViewer.Build.cs
--- a/Engine/Source/Editor/RenderResourceViewer/RenderResourceViewer.Build.cs
+++ b/Engine/Source/Editor/RenderResourceViewer/RenderResourceViewer.Build.cs
@@ -2,6 +2,7 @@
 
 using UnrealBuildTool;
 
+[SupportedTargetTypes(TargetType.Editor)]
 public class RenderResourceViewer : ModuleRules
 {
 	public RenderResourceViewer(ReadOnlyTargetRules Target) : base(Target)
diff --git a/Engine/Source/Editor/RewindDebuggerInterface/RewindDebuggerInterface.Build.cs b/Engine/Source/Editor/RewindDebuggerInterface/RewindDebuggerInterface.Build.cs
--- a/Engine/Source/Editor/RewindDebuggerInterface/RewindDebuggerInterface.Build.cs
+++ b/Engine/Source/Editor/RewindDebuggerInterface/RewindDebuggerInterface.Build.cs
@@ -2,6 +2,7 @@
 
 using UnrealBuildTool;
 
+[SupportedTargetTypes(TargetType.Editor, TargetType.Program)]
 public class RewindDebuggerInterface : ModuleRules
 {
 	public RewindDebuggerInterface(ReadOnlyTargetRules Target) : base(Target)
diff --git a/Engine/Source/Editor/SVGDistanceField/SVGDistanceField.Build.cs b/Engine/Source/Editor/SVGDistanceField/SVGDistanceField.Build.cs
--- a/Engine/Source/Editor/SVGDistanceField/SVGDistanceField.Build.cs
+++ b/Engine/Source/Editor/SVGDistanceField/SVGDistanceField.Build.cs
@@ -2,6 +2,7 @@
 
 using UnrealBuildTool;
 
+[SupportedTargetTypes(TargetType.Editor)]
 public class SVGDistanceField : ModuleRules
 {
 	public SVGDistanceField(ReadOnlyTargetRules Target) : base(Target)
diff --git a/Engine/Source/Editor/SceneDepthPickerMode/SceneDepthPickerMode.Build.cs b/Engine/Source/Editor/SceneDepthPickerMode/SceneDepthPickerMode.Build.cs
--- a/Engine/Source/Editor/SceneDepthPickerMode/SceneDepthPickerMode.Build.cs
+++ b/Engine/Source/Editor/SceneDepthPickerMode/SceneDepthPickerMode.Build.cs
@@ -2,6 +2,7 @@
 
 using UnrealBuildTool;
 
+[SupportedTargetTypes(TargetType.Editor)]
 public class SceneDepthPickerMode : ModuleRules
 {
     public SceneDepthPickerMode(ReadOnlyTargetRules Target) : base(Target)
diff --git a/Engine/Source/Editor/SceneOutliner/SceneOutliner.Build.cs b/Engine/Source/Editor/SceneOutliner/SceneOutliner.Build.cs
--- a/Engine/Source/Editor/SceneOutliner/SceneOutliner.Build.cs
+++ b/Engine/Source/Editor/SceneOutliner/SceneOutliner.Build.cs
@@ -2,6 +2,7 @@
 
 using UnrealBuildTool;
 
+[SupportedTargetTypes(TargetType.Editor)]
 public class SceneOutliner : ModuleRules
 {
 	public SceneOutliner(ReadOnlyTargetRules Target) : base(Target)
diff --git a/Engine/Source/Editor/ScriptableEditorWidgets/ScriptableEditorWidgets.Build.cs b/Engine/Source/Editor/ScriptableEditorWidgets/ScriptableEditorWidgets.Build.cs
--- a/Engine/Source/Editor/ScriptableEditorWidgets/ScriptableEditorWidgets.Build.cs
+++ b/Engine/Source/Editor/ScriptableEditorWidgets/ScriptableEditorWidgets.Build.cs
@@ -2,6 +2,7 @@
 
 using UnrealBuildTool;
 
+[SupportedTargetTypes(TargetType.Editor)]
 public class ScriptableEditorWidgets : ModuleRules
 {
 	public ScriptableEditorWidgets(ReadOnlyTargetRules Target) : base(Target)
diff --git a/Engine/Source/Editor/SequenceRecorder/SequenceRecorder.Build.cs b/Engine/Source/Editor/SequenceRecorder/SequenceRecorder.Build.cs
--- a/Engine/Source/Editor/SequenceRecorder/SequenceRecorder.Build.cs
+++ b/Engine/Source/Editor/SequenceRecorder/SequenceRecorder.Build.cs
@@ -2,6 +2,7 @@
 
 namespace UnrealBuildTool.Rules
 {
+	[SupportedTargetTypes(TargetType.Editor)]
 	public class SequenceRecorder : ModuleRules
 	{
 		public SequenceRecorder(ReadOnlyTargetRules Target) : base(Target)
diff --git a/Engine/Source/Editor/Sequencer/Sequencer.Build.cs b/Engine/Source/Editor/Sequencer/Sequencer.Build.cs
--- a/Engine/Source/Editor/Sequencer/Sequencer.Build.cs
+++ b/Engine/Source/Editor/Sequencer/Sequencer.Build.cs
@@ -2,6 +2,7 @@
 
 using UnrealBuildTool;
 
+[SupportedTargetTypes(TargetType.Editor)]
 public class Sequencer : ModuleRules
 {
 	public Sequencer(ReadOnlyTargetRules Target) : base(Target)
diff --git a/Engine/Source/Editor/SequencerCore/SequencerCore.Build.cs b/Engine/Source/Editor/SequencerCore/SequencerCore.Build.cs
--- a/Engine/Source/Editor/SequencerCore/SequencerCore.Build.cs
+++ b/Engine/Source/Editor/SequencerCore/SequencerCore.Build.cs
@@ -2,6 +2,7 @@
 
 using UnrealBuildTool;
 
+[SupportedTargetTypes(TargetType.Editor)]
 public class SequencerCore : ModuleRules
 {
 	public SequencerCore(ReadOnlyTargetRules Target) : base(Target)
diff --git a/Engine/Source/Editor/SequencerWidgets/SequencerWidgets.Build.cs b/Engine/Source/Editor/SequencerWidgets/SequencerWidgets.Build.cs
--- a/Engine/Source/Editor/SequencerWidgets/SequencerWidgets.Build.cs
+++ b/Engine/Source/Editor/SequencerWidgets/SequencerWidgets.Build.cs
@@ -2,6 +2,7 @@
 
 using UnrealBuildTool;
 
+[SupportedTargetTypes(TargetType.Editor)]
 public class SequencerWidgets : ModuleRules
 {
 	public SequencerWidgets(ReadOnlyTargetRules Target) : base(Target)
diff --git a/Engine/Source/Editor/SerializedRecorderInterface/SerializedRecorderInterface.Build.cs b/Engine/Source/Editor/SerializedRecorderInterface/SerializedRecorderInterface.Build.cs
--- a/Engine/Source/Editor/SerializedRecorderInterface/SerializedRecorderInterface.Build.cs
+++ b/Engine/Source/Editor/SerializedRecorderInterface/SerializedRecorderInterface.Build.cs
@@ -2,6 +2,7 @@
 
 using UnrealBuildTool;
 
+[SupportedTargetTypes(TargetType.Editor)]
 public class SerializedRecorderInterface : ModuleRules
 {
     public SerializedRecorderInterface(ReadOnlyTargetRules Target) : base(Target)
diff --git a/Engine/Source/Editor/SkeletalMeshEditor/SkeletalMeshEditor.Build.cs b/Engine/Source/Editor/SkeletalMeshEditor/SkeletalMeshEditor.Build.cs
--- a/Engine/Source/Editor/SkeletalMeshEditor/SkeletalMeshEditor.Build.cs
+++ b/Engine/Source/Editor/SkeletalMeshEditor/SkeletalMeshEditor.Build.cs
@@ -2,6 +2,7 @@
 
 using UnrealBuildTool;
 
+[SupportedTargetTypes(TargetType.Editor)]
 public class SkeletalMeshEditor : ModuleRules
 {
 	public SkeletalMeshEditor(ReadOnlyTargetRules Target) : base(Target)
diff --git a/Engine/Source/Editor/SkeletonEditor/SkeletonEditor.Build.cs b/Engine/Source/Editor/SkeletonEditor/SkeletonEditor.Build.cs
--- a/Engine/Source/Editor/SkeletonEditor/SkeletonEditor.Build.cs
+++ b/Engine/Source/Editor/SkeletonEditor/SkeletonEditor.Build.cs
@@ -2,6 +2,7 @@
 
 using UnrealBuildTool;
 
+[SupportedTargetTypes(TargetType.Editor)]
 public class SkeletonEditor : ModuleRules
 {
 	public SkeletonEditor(ReadOnlyTargetRules Target) : base(Target)
diff --git a/Engine/Source/Editor/SourceControlWindowExtender/SourceControlWindowExtender.Build.cs b/Engine/Source/Editor/SourceControlWindowExtender/SourceControlWindowExtender.Build.cs
--- a/Engine/Source/Editor/SourceControlWindowExtender/SourceControlWindowExtender.Build.cs
+++ b/Engine/Source/Editor/SourceControlWindowExtender/SourceControlWindowExtender.Build.cs
@@ -2,6 +2,7 @@
 
 using UnrealBuildTool;
 
+[SupportedTargetTypes(TargetType.Editor)]
 public class SourceControlWindowExtender : ModuleRules
 {
 	public SourceControlWindowExtender(ReadOnlyTargetRules Target) : base(Target)
diff --git a/Engine/Source/Editor/SourceControlWindows/SourceControlWindows.Build.cs b/Engine/Source/Editor/SourceControlWindows/SourceControlWindows.Build.cs
--- a/Engine/Source/Editor/SourceControlWindows/SourceControlWindows.Build.cs
+++ b/Engine/Source/Editor/SourceControlWindows/SourceControlWindows.Build.cs
@@ -2,6 +2,7 @@
 
 using UnrealBuildTool;
 
+[SupportedTargetTypes(TargetType.Editor)]
 public class SourceControlWindows : ModuleRules
 {
 	public SourceControlWindows(ReadOnlyTargetRules Target) : base(Target)
diff --git a/Engine/Source/Editor/SparseVolumeTexture/SparseVolumeTexture.Build.cs b/Engine/Source/Editor/SparseVolumeTexture/SparseVolumeTexture.Build.cs
--- a/Engine/Source/Editor/SparseVolumeTexture/SparseVolumeTexture.Build.cs
+++ b/Engine/Source/Editor/SparseVolumeTexture/SparseVolumeTexture.Build.cs
@@ -2,6 +2,7 @@
 
 using UnrealBuildTool;
 
+[SupportedTargetTypes(TargetType.Editor)]
 public class SparseVolumeTexture : ModuleRules
 {
 	public SparseVolumeTexture(ReadOnlyTargetRules Target) : base(Target)
diff --git a/Engine/Source/Editor/StaticMeshEditor/StaticMeshEditor.Build.cs b/Engine/Source/Editor/StaticMeshEditor/StaticMeshEditor.Build.cs
--- a/Engine/Source/Editor/StaticMeshEditor/StaticMeshEditor.Build.cs
+++ b/Engine/Source/Editor/StaticMeshEditor/StaticMeshEditor.Build.cs
@@ -3,6 +3,7 @@
 using UnrealBuildTool;
 using System.IO;
 
+[SupportedTargetTypes(TargetType.Editor)]
 public class StaticMeshEditor : ModuleRules
 {
 	public StaticMeshEditor(ReadOnlyTargetRules Target) : base(Target)
diff --git a/Engine/Source/Editor/StatsViewer/StatsViewer.Build.cs b/Engine/Source/Editor/StatsViewer/StatsViewer.Build.cs
--- a/Engine/Source/Editor/StatsViewer/StatsViewer.Build.cs
+++ b/Engine/Source/Editor/StatsViewer/StatsViewer.Build.cs
@@ -2,6 +2,7 @@
 
 using UnrealBuildTool;
 
+[SupportedTargetTypes(TargetType.Editor)]
 public class StatsViewer : ModuleRules
 {
 	public StatsViewer(ReadOnlyTargetRules Target) : base(Target)
diff --git a/Engine/Source/Editor/StatusBar/StatusBar.Build.cs b/Engine/Source/Editor/StatusBar/StatusBar.Build.cs
--- a/Engine/Source/Editor/StatusBar/StatusBar.Build.cs
+++ b/Engine/Source/Editor/StatusBar/StatusBar.Build.cs
@@ -2,6 +2,7 @@
 
 using UnrealBuildTool;
 
+[SupportedTargetTypes(TargetType.Editor)]
 public class StatusBar : ModuleRules
 {
 	public StatusBar(ReadOnlyTargetRules Target)
diff --git a/Engine/Source/Editor/StorageServerWidgets/StorageServerWidgets.Build.cs b/Engine/Source/Editor/StorageServerWidgets/StorageServerWidgets.Build.cs
--- a/Engine/Source/Editor/StorageServerWidgets/StorageServerWidgets.Build.cs
+++ b/Engine/Source/Editor/StorageServerWidgets/StorageServerWidgets.Build.cs
@@ -2,6 +2,7 @@
 
 using UnrealBuildTool;
 
+[SupportedTargetTypes(TargetType.Editor, TargetType.Program)]
 public class StorageServerWidgets : ModuleRules
 {
 	public StorageServerWidgets(ReadOnlyTargetRules Target)
diff --git a/Engine/Source/Editor/StringTableEditor/StringTableEditor.Build.cs b/Engine/Source/Editor/StringTableEditor/StringTableEditor.Build.cs
--- a/Engine/Source/Editor/StringTableEditor/StringTableEditor.Build.cs
+++ b/Engine/Source/Editor/StringTableEditor/StringTableEditor.Build.cs
@@ -2,6 +2,7 @@
 
 using UnrealBuildTool;
 
+[SupportedTargetTypes(TargetType.Editor)]
 public class StringTableEditor : ModuleRules
 {
 	public StringTableEditor(ReadOnlyTargetRules Target)
diff --git a/Engine/Source/Editor/StructUtilsEditor/StructUtilsEditor.Build.cs b/Engine/Source/Editor/StructUtilsEditor/StructUtilsEditor.Build.cs
--- a/Engine/Source/Editor/StructUtilsEditor/StructUtilsEditor.Build.cs
+++ b/Engine/Source/Editor/StructUtilsEditor/StructUtilsEditor.Build.cs
@@ -2,6 +2,7 @@
 
 namespace UnrealBuildTool.Rules
 {
+	[SupportedTargetTypes(TargetType.Editor)]
 	public class StructUtilsEditor : ModuleRules
 	{
 		public StructUtilsEditor(ReadOnlyTargetRules Target) : base(Target)
diff --git a/Engine/Source/Editor/StructViewer/StructViewer.Build.cs b/Engine/Source/Editor/StructViewer/StructViewer.Build.cs
--- a/Engine/Source/Editor/StructViewer/StructViewer.Build.cs
+++ b/Engine/Source/Editor/StructViewer/StructViewer.Build.cs
@@ -2,6 +2,7 @@
 
 using UnrealBuildTool;
 
+[SupportedTargetTypes(TargetType.Editor)]
 public class StructViewer : ModuleRules
 {
 	public StructViewer(ReadOnlyTargetRules Target) : base(Target)
diff --git a/Engine/Source/Editor/SubobjectDataInterface/SubobjectDataInterface.Build.cs b/Engine/Source/Editor/SubobjectDataInterface/SubobjectDataInterface.Build.cs
--- a/Engine/Source/Editor/SubobjectDataInterface/SubobjectDataInterface.Build.cs
+++ b/Engine/Source/Editor/SubobjectDataInterface/SubobjectDataInterface.Build.cs
@@ -3,6 +3,7 @@
 using UnrealBuildTool;
 using System.IO;
 
+[SupportedTargetTypes(TargetType.Editor)]
 public class SubobjectDataInterface : ModuleRules
 {
 	public SubobjectDataInterface(ReadOnlyTargetRules Target) : base(Target)
diff --git a/Engine/Source/Editor/SubobjectEditor/SubobjectEditor.Build.cs b/Engine/Source/Editor/SubobjectEditor/SubobjectEditor.Build.cs
--- a/Engine/Source/Editor/SubobjectEditor/SubobjectEditor.Build.cs
+++ b/Engine/Source/Editor/SubobjectEditor/SubobjectEditor.Build.cs
@@ -3,6 +3,7 @@
 using UnrealBuildTool;
 using System.IO;
 
+[SupportedTargetTypes(TargetType.Editor)]
 public class SubobjectEditor : ModuleRules
 {
 	public SubobjectEditor(ReadOnlyTargetRules Target) : base(Target)
diff --git a/Engine/Source/Editor/SwarmInterface/SwarmInterface.Build.cs b/Engine/Source/Editor/SwarmInterface/SwarmInterface.Build.cs
--- a/Engine/Source/Editor/SwarmInterface/SwarmInterface.Build.cs
+++ b/Engine/Source/Editor/SwarmInterface/SwarmInterface.Build.cs
@@ -3,6 +3,7 @@
 using UnrealBuildTool;
 using System.IO;
 
+[SupportedTargetTypes(TargetType.Editor)]
 public class SwarmInterface : ModuleRules
 {
 	public SwarmInterface(ReadOnlyTargetRules Target) : base(Target)
diff --git a/Engine/Source/Editor/TextureEditor/TextureEditor.Build.cs b/Engine/Source/Editor/TextureEditor/TextureEditor.Build.cs
--- a/Engine/Source/Editor/TextureEditor/TextureEditor.Build.cs
+++ b/Engine/Source/Editor/TextureEditor/TextureEditor.Build.cs
@@ -2,6 +2,7 @@
 
 using UnrealBuildTool;
 
+[SupportedTargetTypes(TargetType.Editor)]
 public class TextureEditor : ModuleRules
 {
 	public TextureEditor(ReadOnlyTargetRules Target) : base(Target)
diff --git a/Engine/Source/Editor/ToolMenusEditor/ToolMenusEditor.Build.cs b/Engine/Source/Editor/ToolMenusEditor/ToolMenusEditor.Build.cs
--- a/Engine/Source/Editor/ToolMenusEditor/ToolMenusEditor.Build.cs
+++ b/Engine/Source/Editor/ToolMenusEditor/ToolMenusEditor.Build.cs
@@ -2,6 +2,7 @@
 
 namespace UnrealBuildTool.Rules
 {
+	[SupportedTargetTypes(TargetType.Editor)]
 	public class ToolMenusEditor : ModuleRules
 	{
 		public ToolMenusEditor(ReadOnlyTargetRules Target) : base(Target)
diff --git a/Engine/Source/Editor/TurnkeySupport/TurnkeySupport.Build.cs b/Engine/Source/Editor/TurnkeySupport/TurnkeySupport.Build.cs
--- a/Engine/Source/Editor/TurnkeySupport/TurnkeySupport.Build.cs
+++ b/Engine/Source/Editor/TurnkeySupport/TurnkeySupport.Build.cs
@@ -2,6 +2,7 @@
 
 using UnrealBuildTool;
 
+[SupportedTargetTypes(TargetType.Editor, TargetType.Program)]
 public class TurnkeySupport : ModuleRules
 {
 	public TurnkeySupport(ReadOnlyTargetRules Target) : base(Target)
diff --git a/Engine/Source/Editor/UATHelper/UATHelper.Build.cs b/Engine/Source/Editor/UATHelper/UATHelper.Build.cs
--- a/Engine/Source/Editor/UATHelper/UATHelper.Build.cs
+++ b/Engine/Source/Editor/UATHelper/UATHelper.Build.cs
@@ -2,6 +2,7 @@
 
 namespace UnrealBuildTool.Rules
 {
+	[SupportedTargetTypes(TargetType.Editor)]
 	public class UATHelper : ModuleRules
 	{
 		public UATHelper(ReadOnlyTargetRules Target) : base(Target)
diff --git a/Engine/Source/Editor/UMGEditor/UMGEditor.Build.cs b/Engine/Source/Editor/UMGEditor/UMGEditor.Build.cs
--- a/Engine/Source/Editor/UMGEditor/UMGEditor.Build.cs
+++ b/Engine/Source/Editor/UMGEditor/UMGEditor.Build.cs
@@ -2,6 +2,7 @@
 
 using UnrealBuildTool;
 
+[SupportedTargetTypes(TargetType.Editor)]
 public class UMGEditor : ModuleRules
 {
 	public UMGEditor(ReadOnlyTargetRules Target) : base(Target)
diff --git a/Engine/Source/Editor/UndoHistoryEditor/UndoHistoryEditor.Build.cs b/Engine/Source/Editor/UndoHistoryEditor/UndoHistoryEditor.Build.cs
--- a/Engine/Source/Editor/UndoHistoryEditor/UndoHistoryEditor.Build.cs
+++ b/Engine/Source/Editor/UndoHistoryEditor/UndoHistoryEditor.Build.cs
@@ -2,6 +2,7 @@
 
 namespace UnrealBuildTool.Rules
 {
+	[SupportedTargetTypes(TargetType.Editor)]
 	public class UndoHistoryEditor : ModuleRules
 	{
 		public UndoHistoryEditor(ReadOnlyTargetRules Target) : base(Target)
diff --git a/Engine/Source/Editor/UniversalObjectLocatorEditor/UniversalObjectLocatorEditor.Build.cs b/Engine/Source/Editor/UniversalObjectLocatorEditor/UniversalObjectLocatorEditor.Build.cs
--- a/Engine/Source/Editor/UniversalObjectLocatorEditor/UniversalObjectLocatorEditor.Build.cs
+++ b/Engine/Source/Editor/UniversalObjectLocatorEditor/UniversalObjectLocatorEditor.Build.cs
@@ -2,6 +2,7 @@
 
 using UnrealBuildTool;
 
+[SupportedTargetTypes(TargetType.Editor)]
 public class UniversalObjectLocatorEditor : ModuleRules
 {
 	public UniversalObjectLocatorEditor(ReadOnlyTargetRules Target) : base(Target)
diff --git a/Engine/Source/Editor/UnrealEd/UnrealEd.Build.cs b/Engine/Source/Editor/UnrealEd/UnrealEd.Build.cs
--- a/Engine/Source/Editor/UnrealEd/UnrealEd.Build.cs
+++ b/Engine/Source/Editor/UnrealEd/UnrealEd.Build.cs
@@ -389,5 +389,9 @@ public class UnrealEd : ModuleRules
 		{
 			PrivateIncludePathModuleNames.Add("LiveCoding");
 		}
+
+		// Unreal editor adds delegates that are triggered in `UObject` construction and thus
+		// requires transactional support.
+		bDisableAutoRTFMInstrumentation = false;
 	}
 }
diff --git a/Engine/Source/Editor/UserAssetTagsEditor/UserAssetTagsEditor.Build.cs b/Engine/Source/Editor/UserAssetTagsEditor/UserAssetTagsEditor.Build.cs
--- a/Engine/Source/Editor/UserAssetTagsEditor/UserAssetTagsEditor.Build.cs
+++ b/Engine/Source/Editor/UserAssetTagsEditor/UserAssetTagsEditor.Build.cs
@@ -2,6 +2,7 @@
 
 using UnrealBuildTool;
 
+[SupportedTargetTypes(TargetType.Editor)]
 public class UserAssetTagsEditor : ModuleRules
 {
 	public UserAssetTagsEditor(ReadOnlyTargetRules Target) : base(Target)
diff --git a/Engine/Source/Editor/VREditor/VREditor.Build.cs b/Engine/Source/Editor/VREditor/VREditor.Build.cs
--- a/Engine/Source/Editor/VREditor/VREditor.Build.cs
+++ b/Engine/Source/Editor/VREditor/VREditor.Build.cs
@@ -2,6 +2,7 @@
 
 namespace UnrealBuildTool.Rules
 {
+	[SupportedTargetTypes(TargetType.Editor)]
 	public class VREditor : ModuleRules
 	{
         public VREditor(ReadOnlyTargetRules Target) : base(Target)
diff --git a/Engine/Source/Editor/ViewportInteraction/ViewportInteraction.Build.cs b/Engine/Source/Editor/ViewportInteraction/ViewportInteraction.Build.cs
--- a/Engine/Source/Editor/ViewportInteraction/ViewportInteraction.Build.cs
+++ b/Engine/Source/Editor/ViewportInteraction/ViewportInteraction.Build.cs
@@ -4,7 +4,8 @@ using UnrealBuildTool;
 
 namespace UnrealBuildTool.Rules
 {
-    public class ViewportInteraction : ModuleRules
+    [SupportedTargetTypes(TargetType.Editor)]
+	public class ViewportInteraction : ModuleRules
     {
         public ViewportInteraction(ReadOnlyTargetRules Target) : base(Target)
         {
diff --git a/Engine/Source/Editor/ViewportSnapping/ViewportSnapping.Build.cs b/Engine/Source/Editor/ViewportSnapping/ViewportSnapping.Build.cs
--- a/Engine/Source/Editor/ViewportSnapping/ViewportSnapping.Build.cs
+++ b/Engine/Source/Editor/ViewportSnapping/ViewportSnapping.Build.cs
@@ -2,6 +2,7 @@
 
 using UnrealBuildTool;
 
+[SupportedTargetTypes(TargetType.Editor)]
 public class ViewportSnapping : ModuleRules
 {
 	public ViewportSnapping(ReadOnlyTargetRules Target) : base(Target)
diff --git a/Engine/Source/Editor/VirtualTexturingEditor/VirtualTexturingEditor.Build.cs b/Engine/Source/Editor/VirtualTexturingEditor/VirtualTexturingEditor.Build.cs
--- a/Engine/Source/Editor/VirtualTexturingEditor/VirtualTexturingEditor.Build.cs
+++ b/Engine/Source/Editor/VirtualTexturingEditor/VirtualTexturingEditor.Build.cs
@@ -2,6 +2,7 @@
 
 using UnrealBuildTool;
 
+[SupportedTargetTypes(TargetType.Editor)]
 public class VirtualTexturingEditor : ModuleRules
 {
 	public VirtualTexturingEditor(ReadOnlyTargetRules Target) : base(Target)
diff --git a/Engine/Source/Editor/VirtualizationEditor/VirtualizationEditor.Build.cs b/Engine/Source/Editor/VirtualizationEditor/VirtualizationEditor.Build.cs
--- a/Engine/Source/Editor/VirtualizationEditor/VirtualizationEditor.Build.cs
+++ b/Engine/Source/Editor/VirtualizationEditor/VirtualizationEditor.Build.cs
@@ -2,6 +2,7 @@
 
 using UnrealBuildTool;
 
+[SupportedTargetTypes(TargetType.Editor)]
 public class VirtualizationEditor : ModuleRules
 {
 	public VirtualizationEditor(ReadOnlyTargetRules Target)
diff --git a/Engine/Source/Editor/WorkspaceMenuStructure/WorkspaceMenuStructure.Build.cs b/Engine/Source/Editor/WorkspaceMenuStructure/WorkspaceMenuStructure.Build.cs
--- a/Engine/Source/Editor/WorkspaceMenuStructure/WorkspaceMenuStructure.Build.cs
+++ b/Engine/Source/Editor/WorkspaceMenuStructure/WorkspaceMenuStructure.Build.cs
@@ -2,6 +2,7 @@
 
 using UnrealBuildTool;
 
+[SupportedTargetTypes(TargetType.Editor, TargetType.Program)]
 public class WorkspaceMenuStructure : ModuleRules
 {
 	public WorkspaceMenuStructure(ReadOnlyTargetRules Target) : base(Target)
diff --git a/Engine/Source/Editor/WorldBookmark/WorldBookmark.Build.cs b/Engine/Source/Editor/WorldBookmark/WorldBookmark.Build.cs
--- a/Engine/Source/Editor/WorldBookmark/WorldBookmark.Build.cs
+++ b/Engine/Source/Editor/WorldBookmark/WorldBookmark.Build.cs
@@ -2,6 +2,7 @@
 
 using UnrealBuildTool;
 
+[SupportedTargetTypes(TargetType.Editor)]
 public class WorldBookmark : ModuleRules
 {
 	public WorldBookmark(ReadOnlyTargetRules Target) : base(Target)
diff --git a/Engine/Source/Editor/WorldBrowser/WorldBrowser.Build.cs b/Engine/Source/Editor/WorldBrowser/WorldBrowser.Build.cs
--- a/Engine/Source/Editor/WorldBrowser/WorldBrowser.Build.cs
+++ b/Engine/Source/Editor/WorldBrowser/WorldBrowser.Build.cs
@@ -2,6 +2,7 @@
 
 using UnrealBuildTool;
 
+[SupportedTargetTypes(TargetType.Editor)]
 public class WorldBrowser : ModuleRules
 {
     public WorldBrowser(ReadOnlyTargetRules Target) : base(Target)
diff --git a/Engine/Source/Editor/WorldPartitionEditor/WorldPartitionEditor.Build.cs b/Engine/Source/Editor/WorldPartitionEditor/WorldPartitionEditor.Build.cs
--- a/Engine/Source/Editor/WorldPartitionEditor/WorldPartitionEditor.Build.cs
+++ b/Engine/Source/Editor/WorldPartitionEditor/WorldPartitionEditor.Build.cs
@@ -2,6 +2,7 @@
 
 using UnrealBuildTool;
 
+[SupportedTargetTypes(TargetType.Editor)]
 public class WorldPartitionEditor : ModuleRules
 {
     public WorldPartitionEditor(ReadOnlyTargetRules Target) : base(Target)
diff --git a/Engine/Source/Editor/ZenEditor/ZenEditor.Build.cs b/Engine/Source/Editor/ZenEditor/ZenEditor.Build.cs
--- a/Engine/Source/Editor/ZenEditor/ZenEditor.Build.cs
+++ b/Engine/Source/Editor/ZenEditor/ZenEditor.Build.cs
@@ -2,6 +2,7 @@
 
 using UnrealBuildTool;
 
+[SupportedTargetTypes(TargetType.Editor)]
 public class ZenEditor : ModuleRules
 {
 	public ZenEditor(ReadOnlyTargetRules Target)
