- Fix Android JSON generation for ANGLE devices
- Add Initial batch of JSON devices to persistent PreviewJsonDevices folder

#Synced-CL 52535510

diff --git a/Engine/Source/Developer/Android/AndroidDeviceDetection/Private/AndroidDeviceDetectionModule.cpp b/Engine/Source/Developer/Android/AndroidDeviceDetection/Private/AndroidDeviceDetectionModule.cpp
--- a/Engine/Source/Developer/Android/AndroidDeviceDetection/Private/AndroidDeviceDetectionModule.cpp
+++ b/Engine/Source/Developer/Android/AndroidDeviceDetection/Private/AndroidDeviceDetectionModule.cpp
@@ -289,6 +289,16 @@ private:
 				FString GPUVendorString = ExtractNextToken(StartIndex, SurfaceFlingerOutput, TEXT(","));
 				outGPUFamily = ExtractNextToken(StartIndex, SurfaceFlingerOutput, TEXT(","));
 				outGLVersion = ExtractNextToken(StartIndex, SurfaceFlingerOutput, TEXT("\n"));
+
+				// ANGLE-based devices report the GL version as "ANGLE (<GPU Name>)"
+				{
+					FRegexPattern AnglePattern(TEXT("ANGLE\\s*\\((.+)\\)"));
+					FRegexMatcher AngleMatcher(AnglePattern, outGLVersion);
+					if (AngleMatcher.FindNext())
+					{
+						outGPUFamily = AngleMatcher.GetCaptureGroup(1);
+					}
+				}
 			}
 		}
 	}
@@ -1108,17 +1118,6 @@ public:
 
 			// OpenGL ES 3.x
 			bOpenGL3x = DeviceInfo->OpenGLVersionString.Contains(TEXT("OpenGL ES 3"));
-			if (bOpenGL3x)
-			{
-				DeviceSpecs.AndroidProperties.GLES31RHIState.MaxTextureDimensions = 4096;
-				DeviceSpecs.AndroidProperties.GLES31RHIState.MaxShadowDepthBufferSizeX = 2048;
-				DeviceSpecs.AndroidProperties.GLES31RHIState.MaxShadowDepthBufferSizeY = 2048;
-				DeviceSpecs.AndroidProperties.GLES31RHIState.MaxCubeTextureDimensions = 2048;
-				DeviceSpecs.AndroidProperties.GLES31RHIState.SupportsRenderTargetFormat_PF_G8 = true;
-				DeviceSpecs.AndroidProperties.GLES31RHIState.SupportsRenderTargetFormat_PF_FloatRGBA = DeviceInfo->GLESExtensions.Contains(TEXT("GL_EXT_color_buffer_half_float"));
-				DeviceSpecs.AndroidProperties.GLES31RHIState.SupportsMultipleRenderTargets = true;
-			}
-
 			// OpenGL ES 2.0 devices are no longer supported.
 			if (!bOpenGL3x)
 			{
diff --git a/Engine/Source/Editor/PIEPreviewDeviceProfileSelector/Private/PIEPreviewDevice.cpp b/Engine/Source/Editor/PIEPreviewDeviceProfileSelector/Private/PIEPreviewDevice.cpp
--- a/Engine/Source/Editor/PIEPreviewDeviceProfileSelector/Private/PIEPreviewDevice.cpp
+++ b/Engine/Source/Editor/PIEPreviewDeviceProfileSelector/Private/PIEPreviewDevice.cpp
@@ -256,17 +256,12 @@ ERHIFeatureLevel::Type FPIEPreviewDevice::GetPreviewDeviceFeatureLevel() const
 			GConfig->GetBool(TEXT("/Script/IOSRuntimeSettings.IOSRuntimeSettings"), TEXT("bSupportsMetal"), bProjectBuiltForMetal, GEngineIni);
 			GConfig->GetBool(TEXT("/Script/IOSRuntimeSettings.IOSRuntimeSettings"), TEXT("bSupportsMetalMRT"), bProjectBuiltForMRTMetal, GEngineIni);
 
-			const bool bDeviceSupportsMetal = DeviceSpecs->IOSProperties.MetalRHIState.MaxTextureDimensions > 0;
-
 			// not supporting preview for MRT metal 
 			check(!bProjectBuiltForMRTMetal);
 
 			// at least one of these should be valid!
 			check(bProjectBuiltForMetal);
 
-			// if device doesn't support metal the project must have ES enabled.
-			check(bProjectBuiltForMetal && bDeviceSupportsMetal);
-
 			return ERHIFeatureLevel::ES3_1;
 		}
 		case EPIEPreviewDeviceType::Switch:
@@ -354,24 +349,9 @@ void FPIEPreviewDevice::ApplyRHIOverrides() const
 {
 	EShaderPlatform PreviewPlatform = SP_NumPlatforms;
 	ERHIFeatureLevel::Type PreviewFeatureLevel = GetPreviewDeviceFeatureLevel();
-	FPIERHIOverrideState* RHIOverrideState = nullptr;
 
 	switch (DeviceSpecs->DevicePlatform)
 	{
-		case EPIEPreviewDeviceType::Android:
-		{
-			PreviewPlatform = SP_OPENGL_ES3_1_ANDROID;
-			RHIOverrideState = &DeviceSpecs->AndroidProperties.GLES31RHIState;
-		}
-		break;
-
-		case EPIEPreviewDeviceType::IOS:
-		{
-			PreviewPlatform = SP_METAL_ES3_1;
-			RHIOverrideState = &DeviceSpecs->IOSProperties.MetalRHIState;
-		}
-		break;
-
 		case EPIEPreviewDeviceType::Switch:
 		{
 			// apply r.ForwardShading based on Switch bUseForwardShading setting
@@ -416,16 +396,6 @@ void FPIEPreviewDevice::ApplyRHIOverrides() const
 
 		UMaterialShaderQualitySettings::Get()->SetPreviewPlatform(QualityPreviewShaderPlatform);
 	}
-
-	if (RHIOverrideState != nullptr)
-	{
-		GMaxTextureDimensions.SetPreviewOverride(RHIOverrideState->MaxTextureDimensions);
-		GMaxShadowDepthBufferSizeX.SetPreviewOverride(RHIOverrideState->MaxShadowDepthBufferSizeX);
-		GMaxShadowDepthBufferSizeY.SetPreviewOverride(RHIOverrideState->MaxShadowDepthBufferSizeY);
-		GMaxCubeTextureDimensions.SetPreviewOverride(RHIOverrideState->MaxCubeTextureDimensions);
-		GSupportsRenderTargetFormat_PF_FloatRGBA.SetPreviewOverride(RHIOverrideState->SupportsRenderTargetFormat_PF_FloatRGBA);
-		GSupportsRenderTargetFormat_PF_G8.SetPreviewOverride(RHIOverrideState->SupportsRenderTargetFormat_PF_G8);
-	}
 }
 
 bool FPIEPreviewDevice::GetSelectorPropertyValue(const FName& PropertyType, FString& PropertyValueOUT) const
diff --git a/Engine/Source/Editor/PIEPreviewDeviceSpecification/Public/PIEPreviewDeviceSpecification.h b/Engine/Source/Editor/PIEPreviewDeviceSpecification/Public/PIEPreviewDeviceSpecification.h
--- a/Engine/Source/Editor/PIEPreviewDeviceSpecification/Public/PIEPreviewDeviceSpecification.h
+++ b/Engine/Source/Editor/PIEPreviewDeviceSpecification/Public/PIEPreviewDeviceSpecification.h
@@ -10,8 +10,8 @@
 
 // If you modify the logic or structure of the FPIEPreviewDeviceSpecifications
 // Please bump the version for the platform affected and add backwards compatibility to the MakeOutdatedDeviceSpecCompatible function of the platform's Preview Device Profile Selector.
-#define IOS_DEVICE_PROPERTIES_VERSION 0
-#define ANDROID_DEVICE_PROPERTIES_VERSION 0
+#define IOS_DEVICE_PROPERTIES_VERSION 1
+#define ANDROID_DEVICE_PROPERTIES_VERSION 1
 
 UENUM()
 enum class EPIEPreviewDeviceType : uint8
@@ -55,7 +55,7 @@ public:
 };
 
 USTRUCT()
-struct FPIERHIOverrideState
+struct UE_DEPRECATED(5.8, "PIEPreviewDeviceProfileSelector is deprecated and will be removed. Please use the new Preview Json System to preview Devices") FPIERHIOverrideState
 {
 public:
 	GENERATED_USTRUCT_BODY()
@@ -109,12 +109,8 @@ public:
 	UPROPERTY()
 	FString HMDSystemName;
 
-	UPROPERTY()
-	FPIERHIOverrideState GLES31RHIState;
 	UPROPERTY()
 	bool SM5Available = false;
-// 	UPROPERTY()
-// 	FPIERHIOverrideState VulkanRHIState;
 };
 
 USTRUCT()
@@ -131,9 +127,6 @@ public:
 
 	UPROPERTY()
 	float NativeScaleFactor = 0.0f;
-
-	UPROPERTY()
-	FPIERHIOverrideState MetalRHIState;
 };
 
 USTRUCT()
