Fix Static Analysis

#Synced-CL 52543987

diff --git a/Engine/Source/Editor/PIEPreviewDeviceProfileSelector/Private/PIEPreviewDevice.cpp b/Engine/Source/Editor/PIEPreviewDeviceProfileSelector/Private/PIEPreviewDevice.cpp
--- a/Engine/Source/Editor/PIEPreviewDeviceProfileSelector/Private/PIEPreviewDevice.cpp
+++ b/Engine/Source/Editor/PIEPreviewDeviceProfileSelector/Private/PIEPreviewDevice.cpp
@@ -352,6 +352,18 @@ void FPIEPreviewDevice::ApplyRHIOverrides() const
 
 	switch (DeviceSpecs->DevicePlatform)
 	{
+		case EPIEPreviewDeviceType::Android:
+		{
+			PreviewPlatform = SP_OPENGL_ES3_1_ANDROID;
+		}
+		break;
+
+		case EPIEPreviewDeviceType::IOS:
+		{
+			PreviewPlatform = SP_METAL_ES3_1;
+		}
+		break;
+
 		case EPIEPreviewDeviceType::Switch:
 		{
 			// apply r.ForwardShading based on Switch bUseForwardShading setting
